Editor: Add support for external modules to register tutorials with IntroTutorials (no builtins have been moved out yet)

diff --git a/Engine/Source/Editor/IntroTutorials/Private/IntroTutorials.cpp b/Engine/Source/Editor/IntroTutorials/Private/IntroTutorials.cpp
--- a/Engine/Source/Editor/IntroTutorials/Private/IntroTutorials.cpp
+++ b/Engine/Source/Editor/IntroTutorials/Private/IntroTutorials.cpp
@@ -42,6 +42,8 @@ FIntroTutorials::FIntroTutorials()
 	: CurrentObjectClass(nullptr)
 {
 	bEnablePostTutorialSurveys = false;
+
+	bDesireResettingTutorialSeenFlagOnLoad = FParse::Param(FCommandLine::Get(), TEXT("tutorials"));
 }
 
 /** Generates an analytics name for the given tutorial path string */
@@ -172,7 +174,7 @@ void FIntroTutorials::StartupModule()
 		}
 
 		// maybe reset all the "have I seen this once" flags
-		if (FParse::Param(FCommandLine::Get(), TEXT("tutorials")))
+		if (bDesireResettingTutorialSeenFlagOnLoad)
 		{
 			ResetWelcomeTutorials();
 		}
@@ -272,7 +274,7 @@ void FIntroTutorials::SummonBlueprintTutorialHome(UObject* Asset, bool bForceWel
 	check(Tutorial);
 	SummonTutorialWindowForPage(Tutorial->TutorialPath);
 
-	// make sure weve seen this tutorial
+	// make sure we've seen this tutorial
 	GConfig->SetBool(*IntroTutorialConfigSection, *Tutorial->SeenOnceSettingName, true, GEditorGameAgnosticIni);
 }
 
@@ -597,4 +599,25 @@ FString FIntroTutorials::HandleGotoNextTutorial(const FString& InCurrentPage) co
 	return FString();
 }
 
+void FIntroTutorials::RegisterTutorialForAssetEditor(UClass* AssetClass, const FString& TutorialDocPath, const FString& TutorialHasBeenSeenSettingName, const FString& SurveyGUIDString)
+{
+	FWelcomeTutorialProperties TutorialData(TutorialDocPath, TutorialHasBeenSeenSettingName, SurveyGUIDString);
+
+	if (bDesireResettingTutorialSeenFlagOnLoad)
+	{
+		ResetTutorial(TutorialData);
+	}
+
+	TutorialSurveyMap.Add(TutorialDocPath, TutorialData.SurveyGuid);
+
+	AssetEditorTutorialPropertyMap.Add(AssetClass, TutorialData);
+}
+
+void FIntroTutorials::UnregisterTutorialForAssetEditor(UClass* AssetClass)
+{
+	const FWelcomeTutorialProperties& TutorialData = AssetEditorTutorialPropertyMap.FindChecked(AssetClass);
+	TutorialSurveyMap.Remove(TutorialData.TutorialPath);
+	AssetEditorTutorialPropertyMap.Remove(AssetClass);
+}
+
 #undef LOCTEXT_NAMESPACE
\ No newline at end of file
diff --git a/Engine/Source/Editor/IntroTutorials/Private/IntroTutorials.h b/Engine/Source/Editor/IntroTutorials/Private/IntroTutorials.h
--- a/Engine/Source/Editor/IntroTutorials/Private/IntroTutorials.h
+++ b/Engine/Source/Editor/IntroTutorials/Private/IntroTutorials.h
@@ -123,6 +123,12 @@ private:
 	/** Handle linking between tutorials */
 	FString HandleGotoNextTutorial(const FString& InCurrentPage) const;
 
+public:
+
+	// IIntroTutorials interface
+	virtual void RegisterTutorialForAssetEditor(UClass* AssetClass, const FString& TutorialDocPath, const FString& TutorialHasBeenSeenSettingName, const FString& SurveyGUIDString) override;
+	virtual void UnregisterTutorialForAssetEditor(UClass* AssetClass) override;
+	// End of IIntroTutorials interface
 private:
 	/** The tab id of the tutorial tab */
 	static const FName IntroTutorialTabName;
@@ -156,4 +162,7 @@ private:
 
 	/** Map of what tutorial to go to at the end of another */
 	TMap<FString, FString> TutorialChainMap;
+
+	/** should we be clearing the 'have seen this tutorial' flag? (controlled by -tutorials on the command line) */
+	bool bDesireResettingTutorialSeenFlagOnLoad;
 };
\ No newline at end of file
diff --git a/Engine/Source/Editor/IntroTutorials/Public/IIntroTutorials.h b/Engine/Source/Editor/IntroTutorials/Public/IIntroTutorials.h
--- a/Engine/Source/Editor/IntroTutorials/Public/IIntroTutorials.h
+++ b/Engine/Source/Editor/IntroTutorials/Public/IIntroTutorials.h
@@ -34,5 +34,20 @@ public:
 	{
 		return FModuleManager::Get().IsModuleLoaded( "IntroTutorials" );
 	}
+
+	/**
+	 * Adds a tutorial association for a type of asset (will be summoned the first time an asset editor is opened for the specified asset class)
+	 * @param AssetClass		Asset type to register a tutorial for
+	 * @param TutorialDocPath	The path to the tutorial doc (must be rooted in the Shared directory, e.g., "Shared/Tutorials/InPersonaAnimEditorTutorial")
+	 * @param TutorialHasBeenSeenSettingName	The setting name to store whether this editor has been seen or not yet (stored in GEditorGameAgnosticIni in the [IntroTutorials] section)
+	 * @param SurveyGUIDString	Text representation of a GUID, used to uniquely identify this tutorial if required for a survey.
+	 */
+	virtual void RegisterTutorialForAssetEditor(UClass* AssetClass, const FString& TutorialDocPath, const FString& TutorialHasBeenSeenSettingName, const FString& SurveyGUIDString)=0;
+
+	/**
+	 * Removes the tutorial association for the specified class
+	 * @param AssetClass		Asset type to unregister the tutorial from
+	 */
+	virtual void UnregisterTutorialForAssetEditor(UClass* AssetClass)=0;
 };
 
