Interactivity code support for Template and FirstPerson Tutorials

diff --git a/Engine/Source/Editor/IntroTutorials/Private/InteractivityData.cpp b/Engine/Source/Editor/IntroTutorials/Private/InteractivityData.cpp
--- a/Engine/Source/Editor/IntroTutorials/Private/InteractivityData.cpp
+++ b/Engine/Source/Editor/IntroTutorials/Private/InteractivityData.cpp
@@ -1981,6 +1981,47 @@ void PopulateInteractivityData(TSharedPtr<FInteractiveTutorials> InteractiveTuto
 
 	//////////////////////////////////////////////////////////////////////////
 
+	struct FTemplateOverview
+	{
+		static void BindInteractivityData(const FString& InUDNPath, const TSharedRef<class FInteractiveTutorial>& InTutorial)
+		{
+			InTutorial->SetTutorialStyle(ETutorialStyle::Default);
+		}
+	};
+
+	InteractiveTutorials->BindInteractivityData(TEXT("Shared/Tutorials/TemplateTutorials/TemplateOverview"), FBindInteractivityData::CreateStatic(&FTemplateOverview::BindInteractivityData));
+	
+	//////////////////////////////////////////////////////////////////////////
+
+	struct FFirstPersonTemplateTutorial
+	{
+		static void BindInteractivityData(const FString& InUDNPath, const TSharedRef<class FInteractiveTutorial>& InTutorial)
+		{
+			InTutorial->SetTutorialStyle(ETutorialStyle::Default);
+			InTutorial->AddHighlightWidget(TEXT("StageCharOpen1"), TEXT("ContentBrowser"));
+			InTutorial->AddHighlightWidget(TEXT("StageCharLook"), TEXT("GraphEditorPanel"));
+			InTutorial->AddHighlightWidget(TEXT("StageCharVars"), TEXT("MyBlueprintPanel"));
+			InTutorial->AddHighlightWidget(TEXT("StageCharComp1"), TEXT("ComponentsMode"));
+			InTutorial->AddHighlightWidget(TEXT("StageCharComp2"), TEXT("ComponentsPanel"));
+			InTutorial->AddHighlightWidget(TEXT("StageCharGraph1"), TEXT("ComponentsPanel"));
+			InTutorial->AddHighlightWidget(TEXT("StageCharInput"), TEXT("EditMenu"));
+			InTutorial->AddHighlightWidget(TEXT("StageCharPropOpen"), TEXT("DefaultsMode")); 
+
+			InTutorial->AddHighlightWidget(TEXT("StageProjOpen1"), TEXT("ContentBrowser"));			
+			InTutorial->AddHighlightWidget(TEXT("StagePropVars"), TEXT("ComponentsMode"));
+			InTutorial->AddHighlightWidget(TEXT("StageProjComp1"), TEXT("ComponentsPanel"));			
+			
+			InTutorial->AddHighlightWidget(TEXT("StageHUDOpen1"), TEXT("ContentBrowser"));
+			InTutorial->AddHighlightWidget(TEXT("StageHUDGraph1"), TEXT("GraphEditorPanel"));
+
+			InTutorial->AddHighlightWidget(TEXT("StageGameOpen1"), TEXT("ContentBrowser"));
+			InTutorial->AddHighlightWidget(TEXT("StageGameProps1"), TEXT("DefaultsMode"));
+			InTutorial->AddHighlightWidget(TEXT("StageProjProp1"), TEXT("GraphEditorPanel"));
+			
+		}
+	};
+
+	InteractiveTutorials->BindInteractivityData(TEXT("Shared/Tutorials/TemplateTutorials/FirstPerson"), FBindInteractivityData::CreateStatic(&FFirstPersonTemplateTutorial::BindInteractivityData));
 
 }
 
diff --git a/Engine/Source/Editor/IntroTutorials/Private/IntroTutorials.cpp b/Engine/Source/Editor/IntroTutorials/Private/IntroTutorials.cpp
--- a/Engine/Source/Editor/IntroTutorials/Private/IntroTutorials.cpp
+++ b/Engine/Source/Editor/IntroTutorials/Private/IntroTutorials.cpp
@@ -27,7 +27,8 @@ const FString FIntroTutorials::InEditorTutorialPath(TEXT("Shared/Tutorials/inEdi
 const FString FIntroTutorials::WelcomeTutorialPath(TEXT("Shared/Tutorials/UE4Welcome"));
 const FString FIntroTutorials::InEditorGamifiedTutorialPath(TEXT("Shared/Tutorials/inEditorGamifiedTutorial"));
 const FString FIntroTutorials::HomePath(TEXT("Shared/Tutorials"));
-const FString FIntroTutorials::BlueprintHomePath(TEXT("Shared/Tutorials/InBlueprintEditorTutorial"));
+const FString FIntroTutorials::BlueprintHomePath(TEXT("Shared/Tutorials/TemplateTutorials/TemplateOverview"));
+const FString FIntroTutorials::TemplateOverviewPath(TEXT("Shared/Tutorials/UE4Welcome"));
 
 const FWelcomeTutorialProperties FIntroTutorials::UE4WelcomeTutorial(TEXT("Shared/Tutorials/UE4Welcome"), TEXT("SeenUE4Welcome"));	
 const FWelcomeTutorialProperties FIntroTutorials::BlueprintHomeTutorial(TEXT("Shared/Tutorials/InBlueprintEditorTutorial"), TEXT("SeenBlueprintWelcome"));
@@ -37,6 +38,7 @@ const FWelcomeTutorialProperties FIntroTutorials::InterfaceBlueprintWelcomeTutor
 const FWelcomeTutorialProperties FIntroTutorials::LevelScriptBlueprintWelcomeTutorial(TEXT("Shared/Tutorials/LevelBlueprintInterfaceTutorial"), TEXT("SeenBlueprintWelcome_LevelScript"), FString("B061E309-517D-4916-BFCB-E8104C8F4C35"));
 const FWelcomeTutorialProperties FIntroTutorials::AddCodeToProjectWelcomeTutorial(TEXT("Shared/Tutorials/AddCodeToProjectTutorial"), TEXT("SeenAddCodeToProjectWelcome"), FString("D8D9A7E7-68B3-4CBE-80FE-4B88C47B7524"));
 const FWelcomeTutorialProperties FIntroTutorials::MatineeEditorWelcomeTutorial(TEXT("Shared/Tutorials/InMatineeEditorTutorial"), TEXT("SeenMatineeEditorWelcome"), FString("6439C991-A77B-4B52-953D-3F29B1DD1860"));
+const FWelcomeTutorialProperties FIntroTutorials::TemplateOverview(TemplateOverviewPath, TEXT("SeenTemplateOverview"));
 
 FIntroTutorials::FIntroTutorials()
 	: CurrentObjectClass(nullptr)
@@ -112,6 +114,10 @@ void FIntroTutorials::StartupModule()
 		// Add hook for when AddToCodeProject dialog window is opened
 		FGameProjectGenerationModule::Get().OnAddCodeToProjectDialogOpened().AddRaw(this, &FIntroTutorials::OnAddCodeToProjectDialogOpened);
 
+		// Add hook for New Project dialog window is opened
+		//FGameProjectGenerationModule::Get().OnNewProjectProjectDialogOpened().AddRaw(this, &FIntroTutorials::OnNewProjectDialogOpened);
+
+
 		// Add hook for when editor changes modes (e.g. Place/Paint/Landscape/Foliage)
 		GLevelEditorModeTools().OnEditorModeChanged().AddRaw(this, &FIntroTutorials::OnEditorModeChanged);
 
@@ -518,6 +524,7 @@ void FIntroTutorials::ResetWelcomeTutorials() const
 	ResetTutorial(LevelScriptBlueprintWelcomeTutorial);
 	ResetTutorial(AddCodeToProjectWelcomeTutorial);
 	ResetTutorial(MatineeEditorWelcomeTutorial);
+	ResetTutorial(TemplateOverview);
 }
 
 void FIntroTutorials::OnAddCodeToProjectDialogOpened()
@@ -525,6 +532,11 @@ void FIntroTutorials::OnAddCodeToProjectDialogOpened()
 	MaybeOpenWelcomeTutorial(AddCodeToProjectWelcomeTutorial);
 }
 
+void FIntroTutorials::OnNewProjectDialogOpened()
+{
+	MaybeOpenWelcomeTutorial(TemplateOverview);
+}
+
 void FIntroTutorials::OnMatineeEditorOpened()
 {
 	MaybeOpenWelcomeTutorial(MatineeEditorWelcomeTutorial);
diff --git a/Engine/Source/Editor/IntroTutorials/Private/IntroTutorials.h b/Engine/Source/Editor/IntroTutorials/Private/IntroTutorials.h
--- a/Engine/Source/Editor/IntroTutorials/Private/IntroTutorials.h
+++ b/Engine/Source/Editor/IntroTutorials/Private/IntroTutorials.h
@@ -43,6 +43,7 @@ public:
 	static const FString InEditorGamifiedTutorialPath;
 	static const FString HomePath;
 	static const FString BlueprintHomePath;
+	static const FString TemplateOverviewPath;
 
 	static const FWelcomeTutorialProperties UE4WelcomeTutorial;
 	static const FWelcomeTutorialProperties BlueprintHomeTutorial;
@@ -52,6 +53,7 @@ public:
 	static const FWelcomeTutorialProperties LevelScriptBlueprintWelcomeTutorial;
 	static const FWelcomeTutorialProperties AddCodeToProjectWelcomeTutorial;
 	static const FWelcomeTutorialProperties MatineeEditorWelcomeTutorial;
+	static const FWelcomeTutorialProperties TemplateOverview;
 
 	// ctor
 	FIntroTutorials();
@@ -93,6 +95,7 @@ private:
 
 	/** Events to call when editor changes state in various ways */
 	void OnAddCodeToProjectDialogOpened();
+	void OnNewProjectDialogOpened();
 	void OnMatineeEditorOpened();
 	void OnEditorModeChanged(FEdMode* Mode, bool bEnteringMode);
 
diff --git a/Engine/Source/Editor/MainFrame/Private/Menus/MainMenu.cpp b/Engine/Source/Editor/MainFrame/Private/Menus/MainMenu.cpp
--- a/Engine/Source/Editor/MainFrame/Private/Menus/MainMenu.cpp
+++ b/Engine/Source/Editor/MainFrame/Private/Menus/MainMenu.cpp
@@ -263,6 +263,8 @@ TSharedRef<SWidget> FMainMenu::MakeMainMenu( const TSharedPtr<FTabManager>& TabM
 			LOCTEXT("EditMenu_ToolTip", "Open the edit menu"),
 			FNewMenuDelegate::CreateStatic(&FMainMenu::FillEditMenu, Extender, TabManager),
 			"Edit"
+			,
+			FName(TEXT("EditMenu"))
 		);
 
 		// Window
