Several tutorial tweaks:
Added some tags to several BP nodes (for use by tutorial)
Added provision to open relevant editor for highlighted widget
Added support for highlighting blueprints nodes

diff --git a/Engine/Source/Editor/GraphEditor/Private/KismetNodes/SGraphNodeK2Base.cpp b/Engine/Source/Editor/GraphEditor/Private/KismetNodes/SGraphNodeK2Base.cpp
--- a/Engine/Source/Editor/GraphEditor/Private/KismetNodes/SGraphNodeK2Base.cpp
+++ b/Engine/Source/Editor/GraphEditor/Private/KismetNodes/SGraphNodeK2Base.cpp
@@ -44,6 +44,14 @@ void SGraphNodeK2Base::UpdateCompactNode()
 		NodeToolTip = IDocumentation::Get()->CreateToolTip( TAttribute< FText >( this, &SGraphNode::GetNodeTooltip ), NULL, GraphNode->GetDocumentationLink(), GraphNode->GetDocumentationExcerptName() );
 	}
 
+	// Setup a tag for this node
+	FString TagName;
+	if (GraphNode != nullptr)
+	{
+		UBlueprint* Blueprint = FBlueprintEditorUtils::FindBlueprintForGraphChecked(CastChecked<UEdGraph>(GraphNode->GetOuter()));
+		TagName = FString::Printf(TEXT("%s,%s"), *GraphNode->GetName(), *GraphNode->NodeGuid.ToString());		
+	}
+	
 	//
 	//             ______________________
 	//            | (>) L |      | R (>) |
@@ -70,6 +78,7 @@ void SGraphNodeK2Base::UpdateCompactNode()
 			// NODE CONTENT AREA
 			SNew(SOverlay)
 			.ToolTip( NodeToolTip.ToSharedRef() )
+			.Tag(*TagName)
 			+SOverlay::Slot()
 			[
 				SNew(SImage)
diff --git a/Engine/Source/Editor/GraphEditor/Private/KismetNodes/SGraphNodeK2Var.cpp b/Engine/Source/Editor/GraphEditor/Private/KismetNodes/SGraphNodeK2Var.cpp
--- a/Engine/Source/Editor/GraphEditor/Private/KismetNodes/SGraphNodeK2Var.cpp
+++ b/Engine/Source/Editor/GraphEditor/Private/KismetNodes/SGraphNodeK2Var.cpp
@@ -158,6 +158,29 @@ void SGraphNodeK2Var::UpdateGraphNode()
 
 	TSharedPtr<SWidget> ErrorText = SetupErrorReporting();
 
+	// Setup a tag for this node
+	FString TagName;
+	if (GraphNode != nullptr)
+	{
+		// Try to get a suitable name for the tag.
+		if (UK2Node_Variable const* K2Node = Cast<UK2Node_Variable>(GraphNode))
+		{
+			UBlueprint* Blueprint = FBlueprintEditorUtils::FindBlueprintForGraphChecked(CastChecked<UEdGraph>(GraphNode->GetOuter()));
+			TagName = FString::Printf(TEXT("%s,%s"), *K2Node->GetVarNameString(), *GraphNode->NodeGuid.ToString());
+		}
+		else if (UK2Node_Self const* K2Node = Cast<UK2Node_Self>(GraphNode))
+		{
+			UBlueprint* Blueprint = FBlueprintEditorUtils::FindBlueprintForGraphChecked(CastChecked<UEdGraph>(GraphNode->GetOuter()));
+			TagName = FString::Printf(TEXT("Self,%s"), *GraphNode->NodeGuid.ToString());
+		}
+		else
+		{
+			UBlueprint* Blueprint = FBlueprintEditorUtils::FindBlueprintForGraphChecked(CastChecked<UEdGraph>(GraphNode->GetOuter()));
+			TagName = FString::Printf(TEXT("%s,%s"), *GraphNode->GetName(), *GraphNode->NodeGuid.ToString());
+		}
+
+	}
+
 	//             ________________
 	//            | (>) L |  R (>) |
 	//            | (>) E |  I (>) |
@@ -182,6 +205,7 @@ void SGraphNodeK2Var::UpdateGraphNode()
 		+SVerticalBox::Slot()
 		[
 			SNew(SOverlay)
+			.Tag(*TagName)
 			+ SOverlay::Slot()
 			[
 				SNew(SImage)
diff --git a/Engine/Source/Editor/GraphEditor/Private/SGraphNode.cpp b/Engine/Source/Editor/GraphEditor/Private/SGraphNode.cpp
--- a/Engine/Source/Editor/GraphEditor/Private/SGraphNode.cpp
+++ b/Engine/Source/Editor/GraphEditor/Private/SGraphNode.cpp
@@ -697,6 +697,24 @@ void SGraphNode::UpdateGraphNode()
 		SetToolTip(DefaultToolTip);
 	}
 
+	// Setup a tag for this node
+ 	FString TagName;
+ 	if (GraphNode != nullptr)
+ 	{
+		UBlueprint* Blueprint = FBlueprintEditorUtils::FindBlueprintForGraphChecked(CastChecked<UEdGraph>(GraphNode->GetOuter()));
+		FString NodeName;
+		if (Blueprint != nullptr)
+		{
+			NodeName = Blueprint->GetOuter()->GetName();
+		}
+		else
+		{
+			NodeName = NodeTitle->GetHeadTitle().ToString();
+		}
+
+		TagName = FString::Printf(TEXT("%s,%s"), *NodeName, *GraphNode->NodeGuid.ToString());
+ 	}
+	
 	TSharedPtr<SVerticalBox> InnerVerticalBox;
 	this->ContentScale.Bind( this, &SGraphNode::GetContentScale );
 	this->ChildSlot
@@ -710,6 +728,7 @@ void SGraphNode::UpdateGraphNode()
 				SNew(SBorder)
 				.BorderImage( FEditorStyle::GetBrush( "Graph.Node.Body" ) )
 				.Padding(0)
+				.Tag(*TagName)
 				[
 					SAssignNew(InnerVerticalBox, SVerticalBox)
 					+SVerticalBox::Slot()
diff --git a/Engine/Source/Editor/GraphEditor/Private/SGraphNodeComment.cpp b/Engine/Source/Editor/GraphEditor/Private/SGraphNodeComment.cpp
--- a/Engine/Source/Editor/GraphEditor/Private/SGraphNodeComment.cpp
+++ b/Engine/Source/Editor/GraphEditor/Private/SGraphNodeComment.cpp
@@ -96,6 +96,16 @@ void SGraphNodeComment::UpdateGraphNode()
 	// Avoid standard box model too
 	RightNodeBox.Reset();
 	LeftNodeBox.Reset();
+	
+	// Setup a tag for this node
+	FString TagName;
+	if (GraphNode != nullptr)
+	{
+		UBlueprint* Blueprint = FBlueprintEditorUtils::FindBlueprintForGraphChecked(CastChecked<UEdGraph>(GraphNode->GetOuter()));
+		FString NodeName;
+		NodeName = GetEditableNodeTitleAsText().ToString();
+		TagName = FString::Printf(TEXT("%s,%s"), *NodeName, *GraphNode->NodeGuid.ToString());
+	}
 
 	TSharedPtr<SWidget> ErrorText = SetupErrorReporting();
 
@@ -110,6 +120,7 @@ void SGraphNodeComment::UpdateGraphNode()
 			.ColorAndOpacity( FLinearColor::White )
 			.BorderBackgroundColor( this, &SGraphNodeComment::GetCommentBodyColor )
 			.Padding(  FMargin(3.0f) )
+			.Tag(*TagName)
 			[
 				SNew(SVerticalBox)
 				.ToolTipText( this, &SGraphNode::GetNodeTooltip )
diff --git a/Engine/Source/Editor/IntroTutorials/Private/STutorialOverlay.cpp b/Engine/Source/Editor/IntroTutorials/Private/STutorialOverlay.cpp
--- a/Engine/Source/Editor/IntroTutorials/Private/STutorialOverlay.cpp
+++ b/Engine/Source/Editor/IntroTutorials/Private/STutorialOverlay.cpp
@@ -12,6 +12,9 @@ void STutorialOverlay::Construct(const FArguments& InArgs, FTutorialStage* const
 	bIsStandalone = InArgs._IsStandalone;
 	OnClosed = InArgs._OnClosed;
 
+	// Setup the map for opening of closed tabs by highlighted widgets
+	AddTabInfo();
+
 	TSharedPtr<SOverlay> Overlay;
 
 	ChildSlot
@@ -61,6 +64,7 @@ void STutorialOverlay::Construct(const FArguments& InArgs, FTutorialStage* const
 						.OnClosed(OnClosed)
 						.WrapTextAt(WidgetContent.ContentWidth)
 						.Anchor(WidgetContent.WidgetAnchor);
+					OpenBrowserForWidgetAnchor(WidgetContent);
 
 					OverlayCanvas->AddSlot()
 					.Position(TAttribute<FVector2D>::Create(TAttribute<FVector2D>::FGetter::CreateSP(ContentWidget.Get(), &STutorialContent::GetPosition)))
@@ -128,4 +132,40 @@ int32 STutorialOverlay::TraverseWidgets(TSharedRef<SWidget> InWidget, const FGeo
 	}
 
 	return LayerId;
-}
\ No newline at end of file
+}
+
+void STutorialOverlay::OpenBrowserForWidgetAnchor(const FTutorialWidgetContent &WidgetContent)
+{
+	FString IdentString = WidgetContent.WidgetAnchor.WrapperIdentifier.ToString();
+	FString TabString = "";
+		
+	// See if we have a mapping for the widget ident
+	for (TMap<FString, FString>::TConstIterator It(BrowserTabMap); It; ++It)
+	{
+		if (It.Key().StartsWith(IdentString) == true)
+		{
+			TabString = It.Value();
+		}
+	}
+
+	// Open the required tab if we found it in the map
+	if (TabString.IsEmpty() == false)
+	{
+		FLevelEditorModule& LevelEditorModule = FModuleManager::GetModuleChecked<FLevelEditorModule>(TEXT("LevelEditor"));
+		TSharedPtr<FTabManager> LevelEditorTabManager = LevelEditorModule.GetLevelEditorTabManager();
+		LevelEditorTabManager->InvokeTab(FName(*TabString));
+	}
+	
+}
+
+void STutorialOverlay::AddTabInfo()
+{
+	BrowserTabMap.Empty();
+	BrowserTabMap.Add(FString(TEXT("ActorDetails")), FString(TEXT("LevelEditorSelectionDetails")));
+	BrowserTabMap.Add(FString(TEXT("SceneOutliner")), FString(TEXT("LevelEditorSceneOutliner")));
+	BrowserTabMap.Add(FString(TEXT("ContentBrowser")), FString(TEXT("ContentBrowserTab1")));
+	BrowserTabMap.Add(FString(TEXT("ToolsPanel")), FString(TEXT("LevelEditorToolBox")));
+	BrowserTabMap.Add(FString(TEXT("WorldSettings")), FString(TEXT("WorldSettingsTab")));
+	BrowserTabMap.Add(FString(TEXT("EditorViewports")), FString(TEXT("LevelEditorViewport")));	 
+	BrowserTabMap.Add(FString(TEXT("LayerBrowser")), FString(TEXT("LevelEditorLayerBrowser")));	
+}
diff --git a/Engine/Source/Editor/IntroTutorials/Private/STutorialOverlay.h b/Engine/Source/Editor/IntroTutorials/Private/STutorialOverlay.h
--- a/Engine/Source/Editor/IntroTutorials/Private/STutorialOverlay.h
+++ b/Engine/Source/Editor/IntroTutorials/Private/STutorialOverlay.h
@@ -46,6 +46,12 @@ private:
 	/** Recursive function used to re-generate widget geometry and forward the geometry of named widgets onto their respective content */
 	int32 TraverseWidgets(TSharedRef<SWidget> InWidget, const FGeometry& InGeometry, const FSlateRect& MyClippingRect, FSlateWindowElementList& OutDrawElements, int32 LayerId) const;
 
+	/* Opens the browser that the given widget requires if it is not already. */
+	void OpenBrowserForWidgetAnchor(const FTutorialWidgetContent &WidgetContent);
+
+	/* Creates a map of broswers required by the various widget types. */
+	void AddTabInfo();
+
 private:
 	/** Reference to the canvas we use to position our content widgets */
 	TSharedPtr<SCanvas> OverlayCanvas;
@@ -67,4 +73,7 @@ private:
 
 	/** Delegate used to inform widgets of the current window size, so they can auto-adjust layout */
 	FOnCacheWindowSize OnCacheWindowSize;
+
+	/* A map of which tab requires invoking based on the tag name of the asset */
+	TMap< FString, FString >	BrowserTabMap;
 };
\ No newline at end of file
diff --git a/Engine/Source/Editor/IntroTutorials/Private/TutorialStructCustomization.cpp b/Engine/Source/Editor/IntroTutorials/Private/TutorialStructCustomization.cpp
--- a/Engine/Source/Editor/IntroTutorials/Private/TutorialStructCustomization.cpp
+++ b/Engine/Source/Editor/IntroTutorials/Private/TutorialStructCustomization.cpp
@@ -301,8 +301,26 @@ private:
 	{
 		FName WidgetName = NAME_None;
 		TSharedPtr<IPropertyHandle> WrapperIdentifierProperty = StructPropertyHandle->GetChildHandle(GET_MEMBER_NAME_CHECKED(FTutorialContentAnchor, WrapperIdentifier));
-		WrapperIdentifierProperty->GetValue(WidgetName);
-		return FText::FromName(WidgetName);
+
+		// Some of the tags contain a GUID, we dont want to see that here
+		FString WidgetValue;
+		WrapperIdentifierProperty->GetValue(WidgetValue);
+		return MakeFriendlyStringFromName(WidgetValue);
+	}
+
+	FText MakeFriendlyStringFromName(const FString& WidgetName) const
+	{
+		if (WidgetName == "None")
+		{
+			return FText::FromName(*WidgetName);
+		}
+		TArray<FString> Tokens;
+		WidgetName.ParseIntoArray(&Tokens, TEXT(","),true);
+		if (Tokens.Num() == 2)
+		{
+			return FText::FromName(*Tokens[0]);
+		}
+		return FText::FromName(*WidgetName);
 	}
 
 	bool OnIsPicking(FName& OutWidgetNameToHighlight) const
diff --git a/Engine/Source/Editor/LevelEditor/Private/SLevelEditor.cpp b/Engine/Source/Editor/LevelEditor/Private/SLevelEditor.cpp
--- a/Engine/Source/Editor/LevelEditor/Private/SLevelEditor.cpp
+++ b/Engine/Source/Editor/LevelEditor/Private/SLevelEditor.cpp
@@ -686,6 +686,7 @@ TSharedRef<SDockTab> SLevelEditor::SpawnLevelEditorTab( const FSpawnTabArgs& Arg
 					SNew(SBorder)
 					.Padding( 0 )
 					.BorderImage( FEditorStyle::GetBrush("ToolPanel.GroupBorder") )
+					.Tag(TEXT("LayerBrowser"))
 					[
 						LayersModule.CreateLayerBrowser()
 					]
@@ -778,6 +779,7 @@ TSharedRef<SDockTab> SLevelEditor::SpawnLevelEditorTab( const FSpawnTabArgs& Arg
 		return SNew( SDockTab )
 			.Icon( FEditorStyle::GetBrush( "LevelEditor.WorldProperties.Tab" ) )
 			.Label( NSLOCTEXT("LevelEditor", "WorldSettingsTabTitle", "World Settings") )
+			.Tag(TEXT("WorldSettings"))
 			[
 				WorldSettingsView.ToSharedRef()
 			];
diff --git a/Engine/Source/Editor/PropertyEditor/Private/DetailItemNode.cpp b/Engine/Source/Editor/PropertyEditor/Private/DetailItemNode.cpp
--- a/Engine/Source/Editor/PropertyEditor/Private/DetailItemNode.cpp
+++ b/Engine/Source/Editor/PropertyEditor/Private/DetailItemNode.cpp
@@ -135,12 +135,18 @@ void FDetailItemNode::ToggleExpansion()
 
 TSharedRef< ITableRow > FDetailItemNode::GenerateNodeWidget( const TSharedRef<STableViewBase>& OwnerTable, const FDetailColumnSizeData& ColumnSizeData, const TSharedRef<IPropertyUtilities>& PropertyUtilities )
 {
+	FString TagName;
+	if (Customization.IsValidCustomization() && Customization.GetPropertyNode().IsValid() )
+	{
+		TagName = Customization.GetPropertyNode()->GetDisplayName();
+	}
 	if( Customization.HasPropertyNode() && Customization.GetPropertyNode()->AsCategoryNode() )
 	{
-		return 
-			SNew( SDetailCategoryTableRow, AsShared(), OwnerTable )
-			.IsEnabled( IsParentEnabled )
-			.DisplayName( Customization.GetPropertyNode()->GetDisplayName() )
+		return
+			SNew(SDetailCategoryTableRow, AsShared(), OwnerTable)
+			.IsEnabled(IsParentEnabled)
+			.DisplayName(Customization.GetPropertyNode()->GetDisplayName())
+			.Tag(*TagName)
 			.InnerCategory( true );
 	}
 	else
@@ -148,7 +154,8 @@ TSharedRef< ITableRow > FDetailItemNode::GenerateNodeWidget( const TSharedRef<ST
 		return
 			SNew(SDetailSingleItemRow, &Customization, HasMultiColumnWidget(), AsShared(), OwnerTable )
 			.IsEnabled( IsParentEnabled )
-			.ColumnSizeData( ColumnSizeData ) ;
+			.Tag(*TagName)
+			.ColumnSizeData(ColumnSizeData);
 	}
 }
 
diff --git a/Engine/Source/Editor/PropertyEditor/Private/SDetailsView.cpp b/Engine/Source/Editor/PropertyEditor/Private/SDetailsView.cpp
--- a/Engine/Source/Editor/PropertyEditor/Private/SDetailsView.cpp
+++ b/Engine/Source/Editor/PropertyEditor/Private/SDetailsView.cpp
@@ -148,6 +148,7 @@ void SDetailsView::Construct(const FArguments& InArgs)
 				.ContentPadding(0)
 				.ForegroundColor( FSlateColor::UseForeground() )
 				.ButtonStyle( FEditorStyle::Get(), "ToggleButton" )
+				.Tag(TEXT("ViewOtions"))
 				.MenuContent()
 				[
 					DetailViewOptions.MakeWidget()
diff --git a/Engine/Source/Runtime/Slate/Private/Framework/MultiBox/SToolBarComboButtonBlock.cpp b/Engine/Source/Runtime/Slate/Private/Framework/MultiBox/SToolBarComboButtonBlock.cpp
--- a/Engine/Source/Runtime/Slate/Private/Framework/MultiBox/SToolBarComboButtonBlock.cpp
+++ b/Engine/Source/Runtime/Slate/Private/Framework/MultiBox/SToolBarComboButtonBlock.cpp
@@ -107,7 +107,7 @@ void SToolBarComboButtonBlock::BuildMultiBlockWidget(const ISlateStyle* StyleSet
 	// Create the content for our button
 	TSharedRef< SWidget > ButtonContent =
 		SNew( SVerticalBox )
-
+		.Tag(*Label.Get().ToString())
 		// Icon image
 		+ SVerticalBox::Slot()
 		.AutoHeight()
