Stopped tutorial anchor bubbles from showing if they have no text.

diff --git a/Engine/Source/Editor/IntroTutorials/Private/STutorialContent.cpp b/Engine/Source/Editor/IntroTutorials/Private/STutorialContent.cpp
--- a/Engine/Source/Editor/IntroTutorials/Private/STutorialContent.cpp
+++ b/Engine/Source/Editor/IntroTutorials/Private/STutorialContent.cpp
@@ -30,6 +30,16 @@ void STutorialContent::Construct(const FArguments& InArgs, const FTutorialConten
 	BorderIntroAnimation.Play();
 	BorderPulseAnimation.Play();
 
+	if (InContent.Text.IsEmpty() == true)
+	{
+		ChildSlot
+		[
+			SAssignNew(ContentWidget, SBorder)
+			.Visibility(EVisibility::SelfHitTestInvisible)
+		];
+		return;
+	}
+
 	ChildSlot
 	[
 		SAssignNew(ContentWidget, SBorder)
