CIS fix - files I missed

diff --git a/Engine/Source/Editor/IntroTutorials/Private/STutorialEditableText.cpp b/Engine/Source/Editor/IntroTutorials/Private/STutorialEditableText.cpp
new file mode 100644
--- /dev/null
+++ b/Engine/Source/Editor/IntroTutorials/Private/STutorialEditableText.cpp
@@ -0,0 +1,729 @@
+// Copyright 1998-2014 Epic Games, Inc. All Rights Reserved.
+
+#include "IntroTutorialsPrivatePCH.h"
+#include "STutorialEditableText.h"
+#include "IDocumentation.h"
+#include "ISourceCodeAccessModule.h"
+#include "ContentBrowserModule.h"
+
+#define LOCTEXT_NAMESPACE "STutorialEditableText"
+
+namespace TutorialTextHelpers
+{
+
+void OnBrowserLinkClicked(const FSlateHyperlinkRun::FMetadata& Metadata)
+{
+	const FString* Url = Metadata.Find(TEXT("href"));
+	if(Url)
+	{
+		FPlatformProcess::LaunchURL(**Url, nullptr, nullptr);
+	}
+}
+
+
+void OnDocLinkClicked(const FSlateHyperlinkRun::FMetadata& Metadata)
+{
+	const FString* Url = Metadata.Find(TEXT("href"));
+	if(Url)
+	{
+		IDocumentation::Get()->Open(*Url);
+	}
+}
+
+void OnTutorialLinkClicked(const FSlateHyperlinkRun::FMetadata& Metadata)
+{
+	const FString* Url = Metadata.Find(TEXT("href"));
+	if(Url)
+	{
+	//	ParseTutorialLink(*Url);
+	}
+}
+
+static void ParseCodeLink(const FString &InternalLink)
+{
+	// Tokens used by the code parsing. Details in the parse section	
+	static const FString ProjectSpecifier(TEXT("[PROJECT]"));
+	static const FString ProjectRoot(TEXT("[PROJECT]/Source/[PROJECT]/"));
+	static const FString ProjectSuffix(TEXT(".uproject"));
+
+	FString Path;
+	int32 Line = 0;
+	int32 Col = 0;
+
+	TArray<FString> Tokens;
+	InternalLink.ParseIntoArray(&Tokens, TEXT(","), 0);
+	int32 TokenStringsCount = Tokens.Num();
+	if (TokenStringsCount > 0)
+	{
+		Path = Tokens[0];
+	}
+	if (TokenStringsCount > 1)
+	{
+		TTypeFromString<int32>::FromString(Line, *Tokens[1]);
+	}
+	if (TokenStringsCount > 2)
+	{
+		TTypeFromString<int32>::FromString(Col, *Tokens[2]);
+	}
+
+	ISourceCodeAccessModule& SourceCodeAccessModule = FModuleManager::LoadModuleChecked<ISourceCodeAccessModule>("SourceCodeAccess");
+	ISourceCodeAccessor& SourceCodeAccessor = SourceCodeAccessModule.GetAccessor();
+	FString SolutionPath;
+
+	// If we specified generic project specified as the project name try to replace the name with the name of this project
+	if (InternalLink.Contains(ProjectSpecifier) == true)
+	{
+		FString ProjectName = TEXT("Marble");
+		// Try to extract the name of the project
+		FString ProjectPath = FPaths::GetProjectFilePath();
+		if (ProjectPath.EndsWith(ProjectSuffix))
+		{
+			int32 ProjectPathEndIndex;
+			if (ProjectPath.FindLastChar(TEXT('/'), ProjectPathEndIndex) == true)
+			{
+				ProjectName = ProjectPath.RightChop(ProjectPathEndIndex + 1);
+				ProjectName.RemoveFromEnd(*ProjectSuffix);
+			}
+		}
+		// Replace the root path with the name of this project
+		FString RebuiltPath = ProjectRoot + Path;
+		RebuiltPath.ReplaceInline(*ProjectSpecifier, *ProjectName);
+		Path = RebuiltPath;
+	}
+
+	// Finally create the complete path - project name and all
+	int32 PathEndIndex;
+	SolutionPath = FModuleManager::Get().GetSolutionFilepath();
+	if (SolutionPath.FindLastChar(TEXT('/'), PathEndIndex) == true)
+	{
+		SolutionPath = SolutionPath.LeftChop(SolutionPath.Len() - PathEndIndex - 1);
+		SolutionPath += Path;
+		SourceCodeAccessor.OpenFileAtLine(SolutionPath, Line, Col);
+	}
+}
+
+void OnCodeLinkClicked(const FSlateHyperlinkRun::FMetadata& Metadata)
+{
+	const FString* Url = Metadata.Find(TEXT("href"));
+	if(Url)
+	{	
+		ParseCodeLink(*Url);
+	}
+}
+
+static void ParseAssetLink(const FString& InternalLink, const FString& Action)
+{
+	UObject* RequiredObject = FindObject<UObject>(ANY_PACKAGE, *InternalLink);
+	if (RequiredObject != nullptr)
+	{
+		if (Action == TEXT("edit"))
+		{
+			FAssetEditorManager::Get().OpenEditorForAsset(RequiredObject);
+		}
+		else if(Action == TEXT("select"))
+		{
+			FContentBrowserModule& ContentBrowserModule = FModuleManager::Get().LoadModuleChecked<FContentBrowserModule>("ContentBrowser");
+			TArray<UObject*> AssetToBrowse;
+			AssetToBrowse.Add(RequiredObject);
+			ContentBrowserModule.Get().SyncBrowserToAssets(AssetToBrowse);
+		}
+	}
+}
+
+void OnAssetLinkClicked(const FSlateHyperlinkRun::FMetadata& Metadata)
+{
+	const FString* Url = Metadata.Find(TEXT("href"));
+	const FString* Action = Metadata.Find(TEXT("action"));
+	if(Url && Action)
+	{
+		ParseAssetLink(*Url, *Action);
+	}
+}
+
+}
+
+BEGIN_SLATE_FUNCTION_BUILD_OPTIMIZATION
+void STutorialEditableText::Construct(const FArguments& InArgs)
+{
+	OnTextChanged = InArgs._OnTextChanged;
+	OnTextCommitted = InArgs._OnTextCommitted;
+
+	bPickingColor = false;
+	CurrentHyperlinkType = EHyperlinkType::Browser;
+
+	// Define and add the "Roboto" font family
+	TextStyles.AvailableFontFamilies.Emplace(MakeShareable(new FTextStyles::FFontFamily(
+		LOCTEXT("RobotoFontFamily", "Roboto"), 
+		TEXT("Roboto"), 
+		FName(*(FPaths::EngineContentDir() / TEXT("Slate/Fonts/Roboto-Regular.ttf"))),
+		FName(*(FPaths::EngineContentDir() / TEXT("Slate/Fonts/Roboto-Bold.ttf"))),
+		FName(*(FPaths::EngineContentDir() / TEXT("Slate/Fonts/Roboto-Italic.ttf"))),
+		FName(*(FPaths::EngineContentDir() / TEXT("Slate/Fonts/Roboto-BoldItalic.ttf")))
+		)));
+
+	// Set some sensible defaults (these also match the default text style of "TutorialEditableText.Editor.Text"
+	ActiveFontFamily = TextStyles.AvailableFontFamilies[0];
+	FontSize = 11;
+	FontStyle = FTextStyles::EFontStyle::Regular;
+	FontColor = FLinearColor::Black;
+
+	TSharedRef<FRichTextLayoutMarshaller> RichTextMarshaller = FRichTextLayoutMarshaller::Create(
+		TArray<TSharedRef<ITextDecorator>>(), 
+		&FEditorStyle::Get(), 
+		FEditorStyle::Get().GetWidgetStyle<FTextBlockStyle>("TutorialEditableText.Editor.Text")
+		);
+
+	OnBrowserLinkClicked = FSlateHyperlinkRun::FOnClick::CreateStatic(&TutorialTextHelpers::OnBrowserLinkClicked);
+	OnDocLinkClicked = FSlateHyperlinkRun::FOnClick::CreateStatic(&TutorialTextHelpers::OnDocLinkClicked);
+	OnTutorialLinkClicked = FSlateHyperlinkRun::FOnClick::CreateStatic(&TutorialTextHelpers::OnTutorialLinkClicked);
+	OnCodeLinkClicked = FSlateHyperlinkRun::FOnClick::CreateStatic(&TutorialTextHelpers::OnCodeLinkClicked);
+	OnAssetLinkClicked = FSlateHyperlinkRun::FOnClick::CreateStatic(&TutorialTextHelpers::OnAssetLinkClicked);
+	RichTextMarshaller->AppendInlineDecorator(FHyperlinkDecorator::Create(TEXT("browser"), OnBrowserLinkClicked));
+	RichTextMarshaller->AppendInlineDecorator(FHyperlinkDecorator::Create(TEXT("udn"), OnDocLinkClicked));
+	RichTextMarshaller->AppendInlineDecorator(FHyperlinkDecorator::Create(TEXT("tutorial"), OnTutorialLinkClicked));
+	RichTextMarshaller->AppendInlineDecorator(FHyperlinkDecorator::Create(TEXT("code"), OnCodeLinkClicked));
+	RichTextMarshaller->AppendInlineDecorator(FHyperlinkDecorator::Create(TEXT("asset"), OnAssetLinkClicked));
+	RichTextMarshaller->AppendInlineDecorator(FTextStyleDecorator::Create(&TextStyles));
+
+	this->ChildSlot
+	[
+		SNew(SVerticalBox)
+
+		+SVerticalBox::Slot()
+		.AutoHeight()
+		.Padding(FMargin(0.0f, 0.0f, 0.0f, 0.0f))
+		[
+			SAssignNew(RichEditableTextBox, SMultiLineEditableTextBox)
+			.Font(FEditorStyle::Get().GetWidgetStyle<FTextBlockStyle>("TutorialEditableText.Editor.Text").Font)
+			.Text(InArgs._Text)
+			.OnTextChanged(this, &STutorialEditableText::HandleRichEditableTextChanged)
+			.OnTextCommitted(this, &STutorialEditableText::HandleRichEditableTextCommitted)
+			.OnCursorMoved(this, &STutorialEditableText::HandleRichEditableTextCursorMoved)
+			.Marshaller(RichTextMarshaller)
+			.AutoWrapText(true)
+			.Margin(4)
+			.LineHeightPercentage(1.1f)
+		]
+
+		+SVerticalBox::Slot()
+		.AutoHeight()
+		.Padding(FMargin(0.0f, 0.0f, 0.0f, 4.0f))
+		[
+			SNew(SBorder)
+			.Visibility(this, &STutorialEditableText::GetToolbarVisibility)
+			.BorderImage(FEditorStyle::Get().GetBrush("TutorialEditableText.RoundedBackground"))
+			.Padding(FMargin(4))
+			[
+				SNew(SHorizontalBox)
+
+				+SHorizontalBox::Slot()
+				.AutoWidth()
+				[
+					SAssignNew(FontComboBox, SComboBox<TSharedPtr<FTextStyles::FFontFamily>>)
+					.ComboBoxStyle(FEditorStyle::Get(), "TutorialEditableText.Toolbar.ComboBox")
+					.OptionsSource(&TextStyles.AvailableFontFamilies)
+					.OnSelectionChanged(this, &STutorialEditableText::OnActiveFontFamilyChanged)
+					.OnGenerateWidget(this, &STutorialEditableText::GenerateFontFamilyComboEntry)
+					.InitiallySelectedItem(ActiveFontFamily)
+					[
+						SNew(SBox)
+						.Padding(FMargin(0.0f, 0.0f, 2.0f, 0.0f))
+						[
+							SNew(STextBlock)
+							.Text(this, &STutorialEditableText::GetActiveFontFamilyName)
+						]
+					]
+				]
+
+				+SHorizontalBox::Slot()
+				.AutoWidth()
+				[
+					SNew(SBox)
+					.WidthOverride(24)
+					[
+						SNew(SNumericEntryBox<uint8>)
+						.Value(this, &STutorialEditableText::GetFontSize)
+						.OnValueCommitted(this, &STutorialEditableText::SetFontSize)
+					]
+				]
+
+				+SHorizontalBox::Slot()
+				.AutoWidth()
+				.Padding(FMargin(4.0f, 0.0f, 0.0f, 0.0f))
+				[
+					SNew(SCheckBox)
+					.Style(FEditorStyle::Get(), "TutorialEditableText.Toolbar.ToggleButtonCheckbox")
+					.IsChecked(this, &STutorialEditableText::IsFontStyleBold)
+					.OnCheckStateChanged(this, &STutorialEditableText::OnFontStyleBoldChanged)
+					[
+						SNew(SBox)
+						.WidthOverride(24)
+						.HAlign(HAlign_Center)
+						.VAlign(VAlign_Center)
+						[
+							SNew(STextBlock)
+							.TextStyle(FEditorStyle::Get(), "TutorialEditableText.Toolbar.BoldText")
+							.Text(LOCTEXT("BoldLabel", "B"))
+						]
+					]
+				]
+
+				+SHorizontalBox::Slot()
+				.AutoWidth()
+				[
+					SNew(SCheckBox)
+					.Style(FEditorStyle::Get(), "TutorialEditableText.Toolbar.ToggleButtonCheckbox")
+					.IsChecked(this, &STutorialEditableText::IsFontStyleItalic)
+					.OnCheckStateChanged(this, &STutorialEditableText::OnFontStyleItalicChanged)
+					[
+						SNew(SBox)
+						.WidthOverride(24)
+						.HAlign(HAlign_Center)
+						.VAlign(VAlign_Center)
+						[
+							SNew(STextBlock)
+							.TextStyle(FEditorStyle::Get(), "TutorialEditableText.Toolbar.ItalicText")
+							.Text(LOCTEXT("ItalicLabel", "I"))
+						]
+					]
+				]
+
+				+SHorizontalBox::Slot()
+				.Padding(FMargin(4.0f, 0.0f, 0.0f, 0.0f))
+				.AutoWidth()
+				[
+					SNew(SButton)
+					.ButtonStyle(FEditorStyle::Get(), "TutorialEditableText.Toolbar.Button")
+					.OnClicked(this, &STutorialEditableText::OpenFontColorPicker)
+					[
+						SNew(SOverlay)
+
+						+SOverlay::Slot()
+						.Padding(FMargin(0.0f, 0.0f, 0.0f, 4.0f))
+						.HAlign(HAlign_Center)
+						.VAlign(VAlign_Bottom)
+						[
+							SNew(STextBlock)
+							.TextStyle(FEditorStyle::Get(), "TutorialEditableText.Toolbar.BoldText")
+							.Text(LOCTEXT("ColorLabel", "A"))
+						]
+
+						+SOverlay::Slot()
+						.HAlign(HAlign_Center)
+						.VAlign(VAlign_Bottom)
+						[
+							SNew(SColorBlock)
+							.Color(this, &STutorialEditableText::GetFontColor)
+							.Size(FVector2D(20.0f, 6.0f))
+						]
+					]
+				]
+
+				+SHorizontalBox::Slot()
+				.Padding(FMargin(4.0f, 0.0f, 0.0f, 0.0f))
+				.AutoWidth()
+				[
+					SAssignNew(HyperlinkComboButton, SComboButton)
+					.ComboButtonStyle(FEditorStyle::Get(), "TutorialEditableText.Toolbar.ComboButton")
+					.HasDownArrow(false)
+					.OnComboBoxOpened(this, &STutorialEditableText::HandleHyperlinkComboOpened)
+					.ButtonContent()
+					[
+						SNew(SBox)
+						.WidthOverride(20)
+						.HAlign(HAlign_Center)
+						.VAlign(VAlign_Center)
+						[
+							SNew(SImage)
+							.Image(FEditorStyle::Get().GetBrush("TutorialEditableText.Toolbar.HyperlinkImage"))
+						]
+					]
+					.MenuContent()
+					[
+						SNew(SGridPanel)
+						.FillColumn(1, 1.0f)
+
+						+SGridPanel::Slot(0, 0)
+						.HAlign(HAlign_Right)
+						.Padding(FMargin(2.0f))
+						[
+							SNew(STextBlock)
+							.TextStyle(FEditorStyle::Get(), "TutorialEditableText.Toolbar.Text")
+							.Text(LOCTEXT("HyperlinkNameLabel", "Name:"))
+						]
+						+SGridPanel::Slot(1, 0)
+						.Padding(FMargin(2.0f))
+						[
+							SNew(SBox)
+							.WidthOverride(300)
+							[
+								SAssignNew(HyperlinkNameTextBox, SEditableTextBox)
+							]
+						]
+
+						+SGridPanel::Slot(0, 1)
+						.HAlign(HAlign_Right)
+						.Padding(FMargin(2.0f))
+						[
+							SNew(STextBlock)
+							.TextStyle(FEditorStyle::Get(), "TutorialEditableText.Toolbar.Text")
+							.Text(LOCTEXT("HyperlinkURLLabel", "URL:"))
+						]
+						+SGridPanel::Slot(1, 1)
+						.Padding(FMargin(2.0f))
+						[
+							SNew(SBox)
+							.WidthOverride(300)
+							[
+								SAssignNew(HyperlinkURLTextBox, SEditableTextBox)
+							]
+						]
+
+						+SGridPanel::Slot(0, 2)
+						.HAlign(HAlign_Right)
+						.Padding(FMargin(2.0f))
+						.ColumnSpan(2)
+						[
+							SNew(SHorizontalBox)
+							+SHorizontalBox::Slot()
+							.AutoWidth()
+							.VAlign(VAlign_Center)
+							[
+								SNew(SCheckBox)
+								.Style(FEditorStyle::Get(), "RadioButton")
+								.Type(ESlateCheckBoxType::CheckBox)
+								.IsChecked(this, &STutorialEditableText::IsCreatingBrowserLink)
+								.OnCheckStateChanged(this, &STutorialEditableText::OnCheckBrowserLink)
+								[
+									SNew(STextBlock)
+									.TextStyle(FEditorStyle::Get(), "TutorialEditableText.Toolbar.Text")
+									.Text(LOCTEXT("BrowserLinkTypeLabel", "URL"))
+									.ToolTipText(LOCTEXT("BrowserLinkTypeTooltip", "A link that opens a browser window (e.g. http://www.unrealengine.com)"))
+								]
+							]
+							+SHorizontalBox::Slot()
+							.AutoWidth()
+							.VAlign(VAlign_Center)
+							[
+								SNew(SCheckBox)
+								.IsEnabled(false)
+								.Style(FEditorStyle::Get(), "RadioButton")
+								.Type(ESlateCheckBoxType::CheckBox)
+								.IsChecked(this, &STutorialEditableText::IsCreatingUDNLink)
+								.OnCheckStateChanged(this, &STutorialEditableText::OnCheckUDNLink)
+								[
+									SNew(STextBlock)
+									.TextStyle(FEditorStyle::Get(), "TutorialEditableText.Toolbar.Text")
+									.Text(LOCTEXT("UDNLinkTypeLabel", "UDN"))
+									.ToolTipText(LOCTEXT("UDNLinkTypeTooltip", "A link that opens some UDN documentation (e.g. /Engine/Blueprints/UserGuide/Types/ClassBlueprint)"))
+								]
+							]
+							+SHorizontalBox::Slot()
+							.AutoWidth()
+							.VAlign(VAlign_Center)
+							[
+								SNew(SCheckBox)
+								.IsEnabled(false)
+								.Style(FEditorStyle::Get(), "RadioButton")
+								.Type(ESlateCheckBoxType::CheckBox)
+								.IsChecked(this, &STutorialEditableText::IsCreatingAssetLink)
+								.OnCheckStateChanged(this, &STutorialEditableText::OnCheckAssetLink)
+								[
+									SNew(STextBlock)
+									.TextStyle(FEditorStyle::Get(), "TutorialEditableText.Toolbar.Text")
+									.Text(LOCTEXT("AssetLinkTypeLabel", "Asset"))
+									.ToolTipText(LOCTEXT("AssetLinkTypeTooltip", "A link that opens an asset (e.g. Game/StaticMeshes/SphereMesh)"))
+								]
+							]
+							+SHorizontalBox::Slot()
+							.AutoWidth()
+							.VAlign(VAlign_Center)
+							[
+								SNew(SCheckBox)
+								.IsEnabled(false)
+								.Style(FEditorStyle::Get(), "RadioButton")
+								.Type(ESlateCheckBoxType::CheckBox)
+								.IsChecked(this, &STutorialEditableText::IsCreatingCodeLink)
+								.OnCheckStateChanged(this, &STutorialEditableText::OnCheckCodeLink)
+								[
+									SNew(STextBlock)
+									.TextStyle(FEditorStyle::Get(), "TutorialEditableText.Toolbar.Text")
+									.Text(LOCTEXT("CodeLinkTypeLabel", "Code"))
+									.ToolTipText(LOCTEXT("CodeLinkTypeTooltip", "A link that opens code (e.g. Private/SourceFile.cpp,1,1)"))
+								]
+							]
+							+SHorizontalBox::Slot()
+							.AutoWidth()
+							.Padding(FMargin(5.0f, 0.0f, 0.0f, 0.0f))
+							[
+								SNew(SButton)
+								.ButtonStyle(FEditorStyle::Get(), "TutorialEditableText.Toolbar.Button")
+								.OnClicked(this, &STutorialEditableText::HandleInsertBrowserLinkClicked)
+								[
+									SNew(STextBlock)
+									.TextStyle(FEditorStyle::Get(), "TutorialEditableText.Toolbar.Text")
+									.Text(LOCTEXT("HyperlinkInsertLabel", "Insert Hyperlink"))
+								]
+							]
+						]
+					]
+				]
+			]
+		]
+	];
+}
+END_SLATE_FUNCTION_BUILD_OPTIMIZATION
+	
+void STutorialEditableText::HandleRichEditableTextChanged(const FText& Text)
+{
+	OnTextChanged.ExecuteIfBound(Text);
+}
+
+void STutorialEditableText::HandleRichEditableTextCommitted(const FText& Text, ETextCommit::Type Type)
+{
+	OnTextCommitted.ExecuteIfBound(Text, Type);
+}
+
+void STutorialEditableText::HandleRichEditableTextCursorMoved(const FTextLocation& NewCursorPosition )
+{
+	// We can use GetRunUnderCursor to query the style of the text under the cursor
+	// We can then use this to update the toolbar
+	TSharedPtr<const IRun> Run = RichEditableTextBox->GetRunUnderCursor();
+	if(Run.IsValid() && Run->GetRunInfo().Name == TEXT("TextStyle"))
+	{
+		TextStyles.ExplodeRunInfo(Run->GetRunInfo(), ActiveFontFamily, FontSize, FontStyle, FontColor);
+	}
+}
+
+FText STutorialEditableText::GetActiveFontFamilyName() const
+{
+	return ActiveFontFamily->DisplayName;
+}
+
+void STutorialEditableText::OnActiveFontFamilyChanged(TSharedPtr<FTextStyles::FFontFamily> NewValue, ESelectInfo::Type)
+{
+	ActiveFontFamily = NewValue;
+	StyleSelectedText();
+}
+
+TSharedRef<SWidget> STutorialEditableText::GenerateFontFamilyComboEntry(TSharedPtr<FTextStyles::FFontFamily> SourceEntry)
+{
+	return SNew(STextBlock).Text(SourceEntry->DisplayName);
+}
+
+TOptional<uint8> STutorialEditableText::GetFontSize() const
+{
+	return FontSize;
+}
+
+void STutorialEditableText::SetFontSize(uint8 NewValue, ETextCommit::Type)
+{		
+	FontSize = NewValue;
+	StyleSelectedText();
+}
+
+ESlateCheckBoxState::Type STutorialEditableText::IsFontStyleBold() const
+{
+	return (FontStyle & FTextStyles::EFontStyle::Bold) ? ESlateCheckBoxState::Checked : ESlateCheckBoxState::Unchecked;
+}
+
+void STutorialEditableText::OnFontStyleBoldChanged(ESlateCheckBoxState::Type InState)
+{
+	if(InState == ESlateCheckBoxState::Checked)
+	{
+		FontStyle |= FTextStyles::EFontStyle::Bold;
+	}
+	else
+	{
+		FontStyle &= ~FTextStyles::EFontStyle::Bold;
+	}
+	StyleSelectedText();
+}
+
+ESlateCheckBoxState::Type STutorialEditableText::IsFontStyleItalic() const
+{
+	return (FontStyle & FTextStyles::EFontStyle::Italic) ? ESlateCheckBoxState::Checked : ESlateCheckBoxState::Unchecked;
+}
+
+void STutorialEditableText::OnFontStyleItalicChanged(ESlateCheckBoxState::Type InState)
+{
+	if(InState == ESlateCheckBoxState::Checked)
+	{
+		FontStyle |= FTextStyles::EFontStyle::Italic;
+	}
+	else
+	{
+		FontStyle &= ~FTextStyles::EFontStyle::Italic;
+	}
+	StyleSelectedText();
+}
+
+FLinearColor STutorialEditableText::GetFontColor() const
+{
+	return FontColor;
+}
+
+void STutorialEditableText::SetFontColor(FLinearColor NewValue)
+{
+	FontColor = NewValue;
+	StyleSelectedText();
+}
+
+FReply STutorialEditableText::OpenFontColorPicker()
+{
+	FColorPickerArgs PickerArgs;
+	PickerArgs.bOnlyRefreshOnMouseUp = true;
+	PickerArgs.ParentWidget = AsShared();
+	PickerArgs.bUseAlpha = false;
+	PickerArgs.bOnlyRefreshOnOk = false;
+	PickerArgs.OnColorCommitted = FOnLinearColorValueChanged::CreateSP(this, &STutorialEditableText::SetFontColor);
+	PickerArgs.OnColorPickerCancelled = FOnColorPickerCancelled::CreateSP(this, &STutorialEditableText::SetFontColor);
+	PickerArgs.OnColorPickerWindowClosed = FOnWindowClosed::CreateSP(this, &STutorialEditableText::HandleColorPickerWindowClosed);
+	PickerArgs.InitialColorOverride = FontColor;
+
+	OpenColorPicker(PickerArgs);
+
+	bPickingColor = true;
+
+	return FReply::Handled();
+}
+
+void STutorialEditableText::HandleColorPickerWindowClosed(const TSharedRef<SWindow>& /*InWindow*/)
+{
+	bPickingColor = false;
+}
+
+void STutorialEditableText::StyleSelectedText()
+{
+	// Apply the current style to the selected text
+	// If no text is selected, then a new (empty) run will be inserted with the appropriate style
+	const FRunInfo RunInfo = TextStyles.CreateRunInfo(ActiveFontFamily, FontSize, FontStyle, FontColor);
+	const FTextBlockStyle TextBlockStyle = TextStyles.CreateTextBlockStyle(ActiveFontFamily, FontSize, FontStyle, FontColor);
+	RichEditableTextBox->ApplyToSelection(RunInfo, TextBlockStyle);
+	FSlateApplication::Get().SetKeyboardFocus(RichEditableTextBox, EKeyboardFocusCause::SetDirectly);
+}
+
+void STutorialEditableText::HandleHyperlinkComboOpened()
+{
+	// Read any currently selected text, and use this as the default name of the hyperlink
+	FString SelectedText = RichEditableTextBox->GetSelectedText().ToString();
+	for(int32 SelectedTextIndex = 0; SelectedTextIndex < SelectedText.Len(); ++SelectedTextIndex)
+	{
+		if(FChar::IsLinebreak(SelectedText[SelectedTextIndex]))
+		{
+			SelectedText = SelectedText.Left(SelectedTextIndex);
+			break;
+		}
+	}
+	HyperlinkNameTextBox->SetText(FText::FromString(SelectedText));
+
+	// We can use GetRunUnderCursor to query whether the cursor is currently over a hyperlink
+	// If it is, we can use that as the default URL for the hyperlink
+	TSharedPtr<const IRun> Run = RichEditableTextBox->GetRunUnderCursor();
+	if(Run.IsValid() && Run->GetRunInfo().Name == TEXT("a"))
+	{
+		const FString* const URLUnderCursor = Run->GetRunInfo().MetaData.Find(TEXT("href"));
+		HyperlinkURLTextBox->SetText((URLUnderCursor) ? FText::FromString(*URLUnderCursor) : FText());
+	}
+	else
+	{
+		HyperlinkURLTextBox->SetText(FText());
+	}
+}
+
+FReply STutorialEditableText::HandleInsertBrowserLinkClicked()
+{
+	HyperlinkComboButton->SetIsOpen(false);
+
+	const FText& Name = HyperlinkNameTextBox->GetText();
+	const FText& URL = HyperlinkURLTextBox->GetText();
+
+	// Create the correct meta-information for this run, so that valid source rich-text formatting can be generated for it
+	FRunInfo RunInfo(TEXT("a"));
+	RunInfo.MetaData.Add(TEXT("id"), TEXT("browser"));
+	RunInfo.MetaData.Add(TEXT("href"), URL.ToString());
+	RunInfo.MetaData.Add(TEXT("style"), TEXT("TutorialEditableText.Editor.Hyperlink"));
+
+	// Create the new run, and then insert it at the cursor position
+	TSharedRef<FSlateHyperlinkRun> HyperlinkRun = FSlateHyperlinkRun::Create(
+		RunInfo, 
+		MakeShareable(new FString(Name.ToString())), 
+		FEditorStyle::Get().GetWidgetStyle<FHyperlinkStyle>(FName(TEXT("TutorialEditableText.Editor.Hyperlink"))), 
+		OnBrowserLinkClicked
+		);
+	RichEditableTextBox->InsertRunAtCursor(HyperlinkRun);
+
+	return FReply::Handled();
+}
+
+EVisibility STutorialEditableText::GetToolbarVisibility() const
+{
+	return FontComboBox->IsOpen() || HyperlinkComboButton->IsOpen() || bPickingColor || HasKeyboardFocus() || HasFocusedDescendants() ? EVisibility::Visible : EVisibility::Collapsed;
+}
+
+ESlateCheckBoxState::Type STutorialEditableText::IsCreatingBrowserLink() const
+{
+	return CurrentHyperlinkType == EHyperlinkType::Browser ? ESlateCheckBoxState::Checked : ESlateCheckBoxState::Unchecked;
+}
+
+void STutorialEditableText::OnCheckBrowserLink(ESlateCheckBoxState::Type State)
+{
+	if(State == ESlateCheckBoxState::Checked)
+	{
+		CurrentHyperlinkType = EHyperlinkType::Browser;
+	}
+}
+
+ESlateCheckBoxState::Type STutorialEditableText::IsCreatingUDNLink() const
+{
+	return CurrentHyperlinkType == EHyperlinkType::UDN ? ESlateCheckBoxState::Checked : ESlateCheckBoxState::Unchecked;
+}
+
+void STutorialEditableText::OnCheckUDNLink(ESlateCheckBoxState::Type State)
+{
+	if(State == ESlateCheckBoxState::Checked)
+	{
+		CurrentHyperlinkType = EHyperlinkType::UDN;
+	}
+}
+
+ESlateCheckBoxState::Type STutorialEditableText::IsCreatingTutorialLink() const
+{
+	return CurrentHyperlinkType == EHyperlinkType::Tutorial ? ESlateCheckBoxState::Checked : ESlateCheckBoxState::Unchecked;
+}
+
+void STutorialEditableText::OnCheckTutorialLink(ESlateCheckBoxState::Type State)
+{
+	if(State == ESlateCheckBoxState::Checked)
+	{
+		CurrentHyperlinkType = EHyperlinkType::Tutorial;
+	}
+}
+
+ESlateCheckBoxState::Type STutorialEditableText::IsCreatingCodeLink() const
+{
+	return CurrentHyperlinkType == EHyperlinkType::Code ? ESlateCheckBoxState::Checked : ESlateCheckBoxState::Unchecked;
+}
+
+void STutorialEditableText::OnCheckCodeLink(ESlateCheckBoxState::Type State)
+{
+	if(State == ESlateCheckBoxState::Checked)
+	{
+		CurrentHyperlinkType = EHyperlinkType::Code;
+	}
+}
+
+ESlateCheckBoxState::Type STutorialEditableText::IsCreatingAssetLink() const
+{
+	return CurrentHyperlinkType == EHyperlinkType::Asset ? ESlateCheckBoxState::Checked : ESlateCheckBoxState::Unchecked;
+}
+
+void STutorialEditableText::OnCheckAssetLink(ESlateCheckBoxState::Type State)
+{
+	if(State == ESlateCheckBoxState::Checked)
+	{
+		CurrentHyperlinkType = EHyperlinkType::Asset;
+	}
+}
+
+#undef LOCTEXT_NAMESPACE
\ No newline at end of file
diff --git a/Engine/Source/Editor/IntroTutorials/Private/STutorialEditableText.h b/Engine/Source/Editor/IntroTutorials/Private/STutorialEditableText.h
new file mode 100644
--- /dev/null
+++ b/Engine/Source/Editor/IntroTutorials/Private/STutorialEditableText.h
@@ -0,0 +1,357 @@
+// Copyright 1998-2014 Epic Games, Inc. All Rights Reserved.
+
+#pragma once
+
+namespace TutorialTextHelpers
+{
+	extern void OnBrowserLinkClicked(const FSlateHyperlinkRun::FMetadata& Metadata);
+
+	extern void OnDocLinkClicked(const FSlateHyperlinkRun::FMetadata& Metadata);
+
+	extern void OnTutorialLinkClicked(const FSlateHyperlinkRun::FMetadata& Metadata);
+
+	extern void OnCodeLinkClicked(const FSlateHyperlinkRun::FMetadata& Metadata);
+
+	extern void OnAssetLinkClicked(const FSlateHyperlinkRun::FMetadata& Metadata);
+}
+
+namespace EHyperlinkType
+{
+	enum Type
+	{
+		Browser,
+		UDN,
+		Tutorial,
+		Code,
+		Asset
+	};
+}
+
+ /**
+  * This is used in conjunction with the TextStyle decorator to allow arbitrary styling of text within a rich-text editor
+  * This struct defines a set of known font families, as well as providing some utility functions for converting the text style to and from a text layout run
+  */
+struct FTextStyles
+{
+	/** Flags controlling which TTF or OTF font should be picked from the given font family */
+	struct EFontStyle
+	{
+		typedef uint8 Flags;
+		enum Flag
+		{
+			Regular = 0,
+			Bold = 1<<0,
+			Italic = 1<<1,
+		};
+	};
+
+	/** This struct defines a font family, which combines multiple TTF or OTF fonts into a single group, allowing text to be styled as bold or italic */
+	struct FFontFamily
+	{
+		FFontFamily(FText InDisplayName, FName InFamilyName, FName InRegularFont, FName InBoldFont, FName InItalicFont, FName InBoldItalicFont)
+			: DisplayName(InDisplayName)
+			, FamilyName(InFamilyName)
+			, RegularFont(InRegularFont)
+			, BoldFont(InBoldFont)
+			, ItalicFont(InItalicFont)
+			, BoldItalicFont(InBoldItalicFont)
+		{
+		}
+
+		/** Name used for this font family in the UI */
+		FText DisplayName;
+
+		/** Named used to identify this family from the TextStyle decorator */
+		FName FamilyName;
+
+		/** File paths to the fonts on disk */
+		FName RegularFont;
+		FName BoldFont;
+		FName ItalicFont;
+		FName BoldItalicFont;
+	};
+
+	/** Convert the given text style into run meta-information, so that valid source rich-text formatting can be generated for it */
+	static FRunInfo CreateRunInfo(const TSharedPtr<FFontFamily>& InFontFamily, const uint8 InFontSize, const EFontStyle::Flags InFontStyle, const FLinearColor& InFontColor)
+	{
+		FString FontStyleString;
+		if(InFontStyle == EFontStyle::Regular)
+		{
+			FontStyleString = TEXT("Regular");
+		}
+		else
+		{
+			if(InFontStyle & EFontStyle::Bold)
+			{
+				FontStyleString += "Bold";
+			}
+			if(InFontStyle & EFontStyle::Italic)
+			{
+				FontStyleString += "Italic";
+			}
+		}
+
+		FRunInfo RunInfo(TEXT("TextStyle"));
+		RunInfo.MetaData.Add(TEXT("FontFamily"), InFontFamily->FamilyName.ToString());
+		RunInfo.MetaData.Add(TEXT("FontSize"), FString::FromInt(InFontSize));
+		RunInfo.MetaData.Add(TEXT("FontStyle"), FontStyleString);
+		RunInfo.MetaData.Add(TEXT("FontColor"), InFontColor.ToString());
+		return RunInfo;
+	}
+
+	/** Explode some run meta-information back out into its component text style parts */
+	void ExplodeRunInfo(const FRunInfo& InRunInfo, TSharedPtr<FFontFamily>& OutFontFamily, uint8& OutFontSize, EFontStyle::Flags& OutFontStyle, FLinearColor& OutFontColor) const
+	{
+		check(AvailableFontFamilies.Num());
+
+		const FString* const FontFamilyString = InRunInfo.MetaData.Find(TEXT("FontFamily"));
+		if(FontFamilyString)
+		{
+			OutFontFamily = FindFontFamily(FName(**FontFamilyString));
+		}
+		if(!OutFontFamily.IsValid())
+		{
+			OutFontFamily = AvailableFontFamilies[0];
+		}
+
+		OutFontSize = 11;
+		const FString* const FontSizeString = InRunInfo.MetaData.Find(TEXT("FontSize"));
+		if(FontSizeString)
+		{
+			OutFontSize = static_cast<uint8>(FPlatformString::Atoi(**FontSizeString));
+		}
+
+		OutFontStyle = EFontStyle::Regular;
+		const FString* const FontStyleString = InRunInfo.MetaData.Find(TEXT("FontStyle"));
+		if(FontStyleString)
+		{
+			if(*FontStyleString == TEXT("Bold"))
+			{
+				 OutFontStyle = EFontStyle::Bold;
+			}
+			else if(*FontStyleString == TEXT("Italic"))
+			{
+				 OutFontStyle = EFontStyle::Italic;
+			}
+			else if(*FontStyleString == TEXT("BoldItalic"))
+			{
+				 OutFontStyle = EFontStyle::Bold | EFontStyle::Italic;
+			}
+		}
+
+		OutFontColor = FLinearColor::Black;
+		const FString* const FontColorString = InRunInfo.MetaData.Find(TEXT("FontColor"));
+		if(FontColorString && !OutFontColor.InitFromString(*FontColorString))
+		{
+			OutFontColor = FLinearColor::Black;
+		}
+	}
+
+	/** Convert the given text style into a text block style for use by Slate */
+	static FTextBlockStyle CreateTextBlockStyle(const TSharedPtr<FFontFamily>& InFontFamily, const uint8 InFontSize, const EFontStyle::Flags InFontStyle, const FLinearColor& InFontColor)
+	{
+		FName FontName = InFontFamily->RegularFont;
+		if((InFontStyle & EFontStyle::Bold) && (InFontStyle & EFontStyle::Italic))
+		{
+			FontName = InFontFamily->BoldItalicFont;
+		}
+		else if(InFontStyle & EFontStyle::Bold)
+		{
+			FontName = InFontFamily->BoldFont;
+		}
+		else if(InFontStyle & EFontStyle::Italic)
+		{
+			FontName = InFontFamily->ItalicFont;
+		}
+
+		FTextBlockStyle TextBlockStyle;
+		TextBlockStyle.SetFontName(FontName);
+		TextBlockStyle.SetFontSize(InFontSize);
+		TextBlockStyle.SetColorAndOpacity(InFontColor);
+		return TextBlockStyle;
+	}
+
+	/** Convert the given run meta-information into a text block style for use by Slate */
+	FTextBlockStyle CreateTextBlockStyle(const FRunInfo& InRunInfo) const
+	{
+		TSharedPtr<FFontFamily> FontFamily;
+		uint8 FontSize;
+		EFontStyle::Flags FontStyle;
+		FLinearColor FontColor;
+		ExplodeRunInfo(InRunInfo, FontFamily, FontSize, FontStyle, FontColor);
+		return CreateTextBlockStyle(FontFamily, FontSize, FontStyle, FontColor);
+	}
+
+	/** Try and find a font family with the given name */
+	TSharedPtr<FFontFamily> FindFontFamily(const FName InFamilyName) const
+	{
+		const TSharedPtr<FFontFamily>* const FoundFontFamily = AvailableFontFamilies.FindByPredicate([InFamilyName](TSharedPtr<FFontFamily>& Entry) -> bool
+		{
+			return Entry->FamilyName == InFamilyName;
+		});
+		return (FoundFontFamily) ? *FoundFontFamily : nullptr;
+	}
+
+	TArray<TSharedPtr<FFontFamily>> AvailableFontFamilies;
+};
+
+/**
+ * This is a custom decorator used to allow arbitrary styling of text within a rich-text editor
+ * This is required since normal text styling can only work with known styles from a given Slate style-set
+ */
+class FTextStyleDecorator : public ITextDecorator
+{
+public:
+
+	static TSharedRef<FTextStyleDecorator> Create(FTextStyles* const InTextStyles)
+	{
+		return MakeShareable(new FTextStyleDecorator(InTextStyles));
+	}
+
+	virtual ~FTextStyleDecorator()
+	{
+	}
+
+	virtual bool Supports(const FTextRunParseResults& RunParseResult, const FString& Text) const override
+	{
+		return (RunParseResult.Name == TEXT("TextStyle"));
+	}
+
+	virtual TSharedRef<ISlateRun> Create(const FTextRunParseResults& RunParseResult, const FString& OriginalText, const TSharedRef<FString>& InOutModelText, const ISlateStyle* Style) override
+	{
+		FRunInfo RunInfo(RunParseResult.Name);
+		for(const TPair<FString, FTextRange>& Pair : RunParseResult.MetaData)
+		{
+			RunInfo.MetaData.Add(Pair.Key, OriginalText.Mid(Pair.Value.BeginIndex, Pair.Value.EndIndex - Pair.Value.BeginIndex));
+		}
+
+		FTextRange ModelRange;
+		ModelRange.BeginIndex = InOutModelText->Len();
+		*InOutModelText += OriginalText.Mid(RunParseResult.ContentRange.BeginIndex, RunParseResult.ContentRange.EndIndex - RunParseResult.ContentRange.BeginIndex);
+		ModelRange.EndIndex = InOutModelText->Len();
+
+		return FSlateTextRun::Create(RunInfo, InOutModelText, TextStyles->CreateTextBlockStyle(RunInfo), ModelRange);
+	}
+
+private:
+
+	FTextStyleDecorator(FTextStyles* const InTextStyles)
+		: TextStyles(InTextStyles)
+	{
+	}
+
+	FTextStyles* TextStyles;
+};
+
+
+class STutorialEditableText : public SCompoundWidget
+{
+public:
+	SLATE_BEGIN_ARGS( STutorialEditableText ){}
+
+	SLATE_ATTRIBUTE(FText, Text)
+
+	SLATE_ARGUMENT(FOnTextCommitted, OnTextCommitted)
+
+	SLATE_ARGUMENT(FOnTextChanged, OnTextChanged)
+
+	SLATE_END_ARGS()
+
+	/**
+	 * Construct the widget
+	 *
+	 * @param InArgs   Declaration from which to construct the widget
+	 */
+	void Construct(const FArguments& InArgs);
+
+	virtual bool SupportsKeyboardFocus() const override { return true; }
+
+protected:
+	
+	void HandleRichEditableTextChanged(const FText& Text);
+
+	void HandleRichEditableTextCommitted(const FText& Text, ETextCommit::Type Type);
+
+	void HandleRichEditableTextCursorMoved(const FTextLocation& NewCursorPosition );
+
+	FText GetActiveFontFamilyName() const;
+
+	void OnActiveFontFamilyChanged(TSharedPtr<FTextStyles::FFontFamily> NewValue, ESelectInfo::Type);
+
+	TSharedRef<SWidget> GenerateFontFamilyComboEntry(TSharedPtr<FTextStyles::FFontFamily> SourceEntry);
+
+	TOptional<uint8> GetFontSize() const;
+
+	void SetFontSize(uint8 NewValue, ETextCommit::Type);
+
+	ESlateCheckBoxState::Type IsFontStyleBold() const;
+
+	void OnFontStyleBoldChanged(ESlateCheckBoxState::Type InState);
+
+	ESlateCheckBoxState::Type IsFontStyleItalic() const;
+
+	void OnFontStyleItalicChanged(ESlateCheckBoxState::Type InState);
+
+	FLinearColor GetFontColor() const;
+
+	void SetFontColor(FLinearColor NewValue);
+
+	FReply OpenFontColorPicker();
+
+	void HandleColorPickerWindowClosed(const TSharedRef<SWindow>& InWindow);
+
+	void StyleSelectedText();
+
+	void HandleHyperlinkComboOpened();
+
+	FReply HandleInsertBrowserLinkClicked();
+
+	EVisibility GetToolbarVisibility() const;
+
+	ESlateCheckBoxState::Type IsCreatingBrowserLink() const;
+
+	void OnCheckBrowserLink(ESlateCheckBoxState::Type State);
+
+	ESlateCheckBoxState::Type IsCreatingUDNLink() const;
+
+	void OnCheckUDNLink(ESlateCheckBoxState::Type State);
+
+	ESlateCheckBoxState::Type IsCreatingTutorialLink() const;
+
+	void OnCheckTutorialLink(ESlateCheckBoxState::Type State);
+
+	ESlateCheckBoxState::Type IsCreatingCodeLink() const;
+
+	void OnCheckCodeLink(ESlateCheckBoxState::Type State);
+
+	ESlateCheckBoxState::Type IsCreatingAssetLink() const;
+
+	void OnCheckAssetLink(ESlateCheckBoxState::Type State);
+
+protected:
+	TSharedPtr<SMultiLineEditableTextBox> RichEditableTextBox;
+
+	FSlateHyperlinkRun::FOnClick OnBrowserLinkClicked;
+	FSlateHyperlinkRun::FOnClick OnDocLinkClicked;
+	FSlateHyperlinkRun::FOnClick OnTutorialLinkClicked;
+	FSlateHyperlinkRun::FOnClick OnCodeLinkClicked;
+	FSlateHyperlinkRun::FOnClick OnAssetLinkClicked;
+
+	TSharedPtr<SComboButton> HyperlinkComboButton;
+	TSharedPtr<SComboBox<TSharedPtr<FTextStyles::FFontFamily>>> FontComboBox;
+	TSharedPtr<SEditableTextBox> HyperlinkNameTextBox;
+	TSharedPtr<SEditableTextBox> HyperlinkURLTextBox;
+
+	FTextStyles TextStyles;
+
+	TSharedPtr<FTextStyles::FFontFamily> ActiveFontFamily;
+	uint8 FontSize;
+	FTextStyles::EFontStyle::Flags FontStyle;
+	FLinearColor FontColor;
+	bool bPickingColor;
+
+	FOnTextCommitted OnTextCommitted;
+	FOnTextChanged OnTextChanged;
+
+	EHyperlinkType::Type CurrentHyperlinkType;
+};
