Fix CIS

diff --git a/Engine/Source/Editor/IntroTutorials/Private/STutorialEditableText.h b/Engine/Source/Editor/IntroTutorials/Private/STutorialEditableText.h
--- a/Engine/Source/Editor/IntroTutorials/Private/STutorialEditableText.h
+++ b/Engine/Source/Editor/IntroTutorials/Private/STutorialEditableText.h
@@ -217,7 +217,7 @@ public:
 		return (RunParseResult.Name == TEXT("TextStyle"));
 	}
 
-	virtual TSharedRef<ISlateRun> Create(const FTextRunParseResults& RunParseResult, const FString& OriginalText, const TSharedRef<FString>& InOutModelText, const ISlateStyle* Style) override
+	virtual TSharedRef<ISlateRun> Create(const TSharedRef<class FTextLayout>& TextLayout, const FTextRunParseResults& RunParseResult, const FString& OriginalText, const TSharedRef< FString >& InOutModelText, const ISlateStyle* Style) override
 	{
 		FRunInfo RunInfo(RunParseResult.Name);
 		for(const TPair<FString, FTextRange>& Pair : RunParseResult.MetaData)
