Tutorial widget content can now disable highlighting.

Also changed content visibility for Rich Text so that widgets do not obscure those beneath them.

diff --git a/Engine/Source/Editor/IntroTutorials/Private/STutorialContent.cpp b/Engine/Source/Editor/IntroTutorials/Private/STutorialContent.cpp
--- a/Engine/Source/Editor/IntroTutorials/Private/STutorialContent.cpp
+++ b/Engine/Source/Editor/IntroTutorials/Private/STutorialContent.cpp
@@ -237,6 +237,7 @@ TSharedRef<SWidget> STutorialContent::GenerateContentWidget(const FTutorialConte
 			Decorators.Add(FTextStyleDecorator::Create(&TextStyles));
 
 			return SNew(SRichTextBlock)
+					.Visibility(EVisibility::HitTestInvisible)
 					.TextStyle(FEditorStyle::Get(), "TutorialEditableText.Editor.Text")
 					.DecoratorStyleSet(&FEditorStyle::Get())
 					.Decorators(Decorators)
diff --git a/Engine/Source/Editor/IntroTutorials/Private/TutorialStructCustomization.cpp b/Engine/Source/Editor/IntroTutorials/Private/TutorialStructCustomization.cpp
--- a/Engine/Source/Editor/IntroTutorials/Private/TutorialStructCustomization.cpp
+++ b/Engine/Source/Editor/IntroTutorials/Private/TutorialStructCustomization.cpp
@@ -408,16 +408,42 @@ TSharedRef<IPropertyTypeCustomization> FTutorialContentAnchorCustomization::Make
 
 void FTutorialContentAnchorCustomization::CustomizeHeader( TSharedRef<class IPropertyHandle> InStructPropertyHandle, class FDetailWidgetRow& HeaderRow, IPropertyTypeCustomizationUtils& StructCustomizationUtils )
 {
+	TSharedPtr<IPropertyHandle> DrawHighlightProperty = InStructPropertyHandle->GetChildHandle(GET_MEMBER_NAME_CHECKED(FTutorialContentAnchor, bDrawHighlight));
+
 	HeaderRow
 	.NameContent()
 	[
-		InStructPropertyHandle->CreatePropertyNameWidget()
+		SNew(SVerticalBox)
+		+SVerticalBox::Slot()
+		.AutoHeight()
+		.Padding(0.0f, 2.0f)
+		[
+			InStructPropertyHandle->CreatePropertyNameWidget()
+		]
+		+SVerticalBox::Slot()
+		.AutoHeight()
+		.Padding(0.0f, 2.0f)
+		[
+			DrawHighlightProperty->CreatePropertyNameWidget()
+		]
 	]
 	.ValueContent()
 	.MinDesiredWidth(250.0f)
 	.MaxDesiredWidth(500.0f)
 	[
-		SNew(SWidgetPicker, InStructPropertyHandle, &StructCustomizationUtils)
+		SNew(SVerticalBox)
+		+SVerticalBox::Slot()
+		.AutoHeight()
+		.Padding(0.0f, 2.0f)
+		[
+			SNew(SWidgetPicker, InStructPropertyHandle, &StructCustomizationUtils)
+		]
+		+ SVerticalBox::Slot()
+		.AutoHeight()
+		.Padding(0.0f, 2.0f)
+		[
+			DrawHighlightProperty->CreatePropertyValueWidget()
+		]
 	];
 }
 
