Use the Asset Registry to discover tutorials

Fixed visibility issues where hyperlinks could not be clicked.
Implemented BP functions for tutorial stage completion etc.

diff --git a/Engine/Source/Editor/IntroTutorials/IntroTutorials.Build.cs b/Engine/Source/Editor/IntroTutorials/IntroTutorials.Build.cs
--- a/Engine/Source/Editor/IntroTutorials/IntroTutorials.Build.cs
+++ b/Engine/Source/Editor/IntroTutorials/IntroTutorials.Build.cs
@@ -39,7 +39,8 @@ namespace UnrealBuildTool.Rules
 					"AssetTools",
 					"SourceCodeAccess",
 					"ContentBrowser",
-					"LevelEditor"
+					"LevelEditor",
+                    "AssetRegistry"
 				}
 			);
 
diff --git a/Engine/Source/Editor/IntroTutorials/Private/EditorTutorial.cpp b/Engine/Source/Editor/IntroTutorials/Private/EditorTutorial.cpp
--- a/Engine/Source/Editor/IntroTutorials/Private/EditorTutorial.cpp
+++ b/Engine/Source/Editor/IntroTutorials/Private/EditorTutorial.cpp
@@ -9,19 +9,38 @@ UEditorTutorial::UEditorTutorial(const class FPostConstructInitializeProperties&
 }
 
 
-void UEditorTutorial::GotoNextTutorialStage()
+void UEditorTutorial::GoToNextTutorialStage()
 {
-
+	FIntroTutorials& IntroTutorials = FModuleManager::GetModuleChecked<FIntroTutorials>(TEXT("IntroTutorials"));
+	IntroTutorials.GoToNextStage(nullptr);
 }
 
 
-void UEditorTutorial::GotoPreviousTutorialStage()
+void UEditorTutorial::GoToPreviousTutorialStage()
 {
+	FIntroTutorials& IntroTutorials = FModuleManager::GetModuleChecked<FIntroTutorials>(TEXT("IntroTutorials"));
+	IntroTutorials.GoToPreviousStage();
+}
+
 
+void UEditorTutorial::BeginTutorial(UEditorTutorial* TutorialToStart, bool bRestart)
+{
+	FIntroTutorials& IntroTutorials = FModuleManager::GetModuleChecked<FIntroTutorials>(TEXT("IntroTutorials"));
+	IntroTutorials.LaunchTutorial(TutorialToStart, bRestart);
 }
 
 
-void UEditorTutorial::BeginTutorial(UEditorTutorial* TutorialToStart)
+void UEditorTutorial::HandleTutorialStageStarted(FName StageName)
 {
+	FEditorScriptExecutionGuard ScriptGuard;
+	OnTutorialStageStarted(StageName);
+}
 
+
+void UEditorTutorial::HandleTutorialStageEnded(FName StageName)
+{
+	FEditorScriptExecutionGuard ScriptGuard;
+	OnTutorialStageEnded(StageName);
 }
+
+
diff --git a/Engine/Source/Editor/IntroTutorials/Private/EditorTutorialSettings.h b/Engine/Source/Editor/IntroTutorials/Private/EditorTutorialSettings.h
--- a/Engine/Source/Editor/IntroTutorials/Private/EditorTutorialSettings.h
+++ b/Engine/Source/Editor/IntroTutorials/Private/EditorTutorialSettings.h
@@ -28,10 +28,6 @@ class UEditorTutorialSettings : public UObject
 	UPROPERTY(Config, EditAnywhere, Category="Tutorials")
 	TArray<FTutorialCategory> Categories;
 
-	/** The tutorials available in the editor */
-	UPROPERTY(Config, EditAnywhere, Category="Tutorials", meta=(MetaClass="EditorTutorial"))
-	TArray<FStringClassReference> Tutorials;
-
 	UPROPERTY(Config)
 	TArray<FTutorialProgress> TutorialsProgress;
 
diff --git a/Engine/Source/Editor/IntroTutorials/Private/IntroTutorials.cpp b/Engine/Source/Editor/IntroTutorials/Private/IntroTutorials.cpp
--- a/Engine/Source/Editor/IntroTutorials/Private/IntroTutorials.cpp
+++ b/Engine/Source/Editor/IntroTutorials/Private/IntroTutorials.cpp
@@ -748,4 +748,21 @@ void FIntroTutorials::LaunchTutorial(UEditorTutorial* InTutorial, bool bInRestar
 	}
 }
 
+void FIntroTutorials::GoToPreviousStage()
+{
+	if (TutorialRoot.IsValid())
+	{
+		TutorialRoot->GoToPreviousStage();
+	}
+}
+
+void FIntroTutorials::GoToNextStage(TWeakPtr<SWindow> InNavigationWindow)
+{
+	if (TutorialRoot.IsValid())
+	{
+		TutorialRoot->GoToNextStage(InNavigationWindow);
+	}
+}
+
+
 #undef LOCTEXT_NAMESPACE
\ No newline at end of file
diff --git a/Engine/Source/Editor/IntroTutorials/Private/IntroTutorials.h b/Engine/Source/Editor/IntroTutorials/Private/IntroTutorials.h
--- a/Engine/Source/Editor/IntroTutorials/Private/IntroTutorials.h
+++ b/Engine/Source/Editor/IntroTutorials/Private/IntroTutorials.h
@@ -68,6 +68,10 @@ public:
 	/** Get the delegate used to check for whether we are picking widgets */
 	FOnIsPicking& OnIsPicking();
 
+	void GoToPreviousStage();
+
+	void GoToNextStage(TWeakPtr<SWindow> InNavigationWindow);
+
 private:
 
 	/** IModuleInterface implementation */
diff --git a/Engine/Source/Editor/IntroTutorials/Private/SEditorTutorials.cpp b/Engine/Source/Editor/IntroTutorials/Private/SEditorTutorials.cpp
--- a/Engine/Source/Editor/IntroTutorials/Private/SEditorTutorials.cpp
+++ b/Engine/Source/Editor/IntroTutorials/Private/SEditorTutorials.cpp
@@ -93,6 +93,11 @@ void SEditorTutorials::ShowBrowser(const FString& InFilter)
 	HandleHomeClicked();
 }
 
+bool SEditorTutorials::IsNavigationVisible() const
+{
+	return bShowNavigation;
+}
+
 EVisibility SEditorTutorials::GetBrowserVisibility() const
 {
 	return bBrowserVisible && OnGetCurrentTutorial.Execute() == nullptr ? EVisibility::Visible : EVisibility::Collapsed;
diff --git a/Engine/Source/Editor/IntroTutorials/Private/SEditorTutorials.h b/Engine/Source/Editor/IntroTutorials/Private/SEditorTutorials.h
--- a/Engine/Source/Editor/IntroTutorials/Private/SEditorTutorials.h
+++ b/Engine/Source/Editor/IntroTutorials/Private/SEditorTutorials.h
@@ -48,8 +48,11 @@ class SEditorTutorials : public SCompoundWidget
 	/** Get parent window for this widget */
 	TSharedPtr<SWindow> GetParentWindow() const;
 
-	/** Rebuild content according ot the current tutorial state */
+	/** Rebuild content according to the current tutorial state */
 	void RebuildCurrentContent();
+	
+	/** @return whether the navigation controls are currently visible */
+	bool IsNavigationVisible() const;
 
 private:
 	/** Handle whether we should display the browser */
diff --git a/Engine/Source/Editor/IntroTutorials/Private/STutorialContent.cpp b/Engine/Source/Editor/IntroTutorials/Private/STutorialContent.cpp
--- a/Engine/Source/Editor/IntroTutorials/Private/STutorialContent.cpp
+++ b/Engine/Source/Editor/IntroTutorials/Private/STutorialContent.cpp
@@ -237,7 +237,7 @@ TSharedRef<SWidget> STutorialContent::GenerateContentWidget(const FTutorialConte
 			Decorators.Add(FTextStyleDecorator::Create(&TextStyles));
 
 			return SNew(SRichTextBlock)
-					.Visibility(EVisibility::HitTestInvisible)
+					.Visibility(EVisibility::SelfHitTestInvisible)
 					.TextStyle(FEditorStyle::Get(), "TutorialEditableText.Editor.Text")
 					.DecoratorStyleSet(&FEditorStyle::Get())
 					.Decorators(Decorators)
diff --git a/Engine/Source/Editor/IntroTutorials/Private/STutorialRoot.cpp b/Engine/Source/Editor/IntroTutorials/Private/STutorialRoot.cpp
--- a/Engine/Source/Editor/IntroTutorials/Private/STutorialRoot.cpp
+++ b/Engine/Source/Editor/IntroTutorials/Private/STutorialRoot.cpp
@@ -65,6 +65,11 @@ void STutorialRoot::Tick(const FGeometry& AllottedGeometry, const double InCurre
 	{
 		MaybeAddOverlay(Window);
 	}
+
+	if (CurrentTutorial != nullptr)
+	{
+		CurrentTutorial->HandleTickCurrentStage(CurrentTutorial->Stages[CurrentTutorialStage].Name);
+	}
 }
 
 void STutorialRoot::SummonTutorialBrowser(TSharedRef<SWindow> InWindow, const FString& InFilter)
@@ -90,7 +95,15 @@ void STutorialRoot::LaunchTutorial(UEditorTutorial* InTutorial, bool bInRestart,
 		{
 			if(TutorialWidget.Value.IsValid())
 			{
-				bool bIsNavigationWindow = InNavigationWindow.IsValid() && TutorialWidget.Value.Pin()->GetParentWindow() == InNavigationWindow.Pin();
+				bool bIsNavigationWindow = false;
+				if (!InNavigationWindow.IsValid())
+				{
+					bIsNavigationWindow = TutorialWidget.Value.Pin()->IsNavigationVisible();
+				}
+				else
+				{
+					bIsNavigationWindow = (TutorialWidget.Value.Pin()->GetParentWindow() == InNavigationWindow.Pin());
+				}
 				TutorialWidget.Value.Pin()->LaunchTutorial(bIsNavigationWindow);
 			}
 		}
@@ -168,7 +181,18 @@ void STutorialRoot::GoToPreviousStage()
 {
 	if(CurrentTutorial != nullptr)
 	{
+		int32 PreviousTutorialStage = CurrentTutorialStage;
+		if (CurrentTutorialStage > 0)
+		{
+			CurrentTutorial->HandleTutorialStageEnded(CurrentTutorial->Stages[CurrentTutorialStage].Name);
+		}
+
 		CurrentTutorialStage = FMath::Max(0, CurrentTutorialStage - 1);
+
+		if (PreviousTutorialStage != CurrentTutorialStage)
+		{
+			CurrentTutorial->HandleTutorialStageStarted(CurrentTutorial->Stages[CurrentTutorialStage].Name);
+		}
 	}
 }
 
@@ -176,6 +200,10 @@ void STutorialRoot::GoToNextStage(TWeakPtr<SWindow> InNavigationWindow)
 {
 	if(CurrentTutorial != nullptr)
 	{
+		UEditorTutorial* PreviousTutorial = CurrentTutorial;
+		int32 PreviousTutorialStage = CurrentTutorialStage;
+		CurrentTutorial->HandleTutorialStageEnded(CurrentTutorial->Stages[CurrentTutorialStage].Name);
+
 		if(CurrentTutorialStage + 1 >= CurrentTutorial->Stages.Num() && FName(*CurrentTutorial->NextTutorial.AssetLongPathname) != NAME_None)
 		{
 			TSubclassOf<UEditorTutorial> NextTutorialClass = LoadClass<UEditorTutorial>(NULL, *CurrentTutorial->NextTutorial.AssetLongPathname, NULL, LOAD_None, NULL);
@@ -193,6 +221,11 @@ void STutorialRoot::GoToNextStage(TWeakPtr<SWindow> InNavigationWindow)
 			CurrentTutorialStage = FMath::Min(CurrentTutorialStage + 1, CurrentTutorial->Stages.Num() - 1);
 			GetMutableDefault<UEditorTutorialSettings>()->RecordProgress(CurrentTutorial, CurrentTutorialStage);
 		}
+
+		if (CurrentTutorial != nullptr && (CurrentTutorial != PreviousTutorial || CurrentTutorialStage != PreviousTutorialStage))
+		{
+			CurrentTutorial->HandleTutorialStageStarted(CurrentTutorial->Stages[CurrentTutorialStage].Name);
+		}
 	}
 }
 
diff --git a/Engine/Source/Editor/IntroTutorials/Private/STutorialsBrowser.cpp b/Engine/Source/Editor/IntroTutorials/Private/STutorialsBrowser.cpp
--- a/Engine/Source/Editor/IntroTutorials/Private/STutorialsBrowser.cpp
+++ b/Engine/Source/Editor/IntroTutorials/Private/STutorialsBrowser.cpp
@@ -6,6 +6,7 @@
 #include "STutorialContent.h"
 #include "TutorialSettings.h"
 #include "EditorTutorialSettings.h"
+#include "AssetRegistryModule.h"
 
 #define LOCTEXT_NAMESPACE "TutorialsBrowser"
 
@@ -63,7 +64,7 @@ public:
 				]
 				+SHorizontalBox::Slot()
 				.FillWidth(1.0f)
-				.MaxWidth(400.0f)
+				.MaxWidth(600.0f)
 				.HAlign(HAlign_Left)
 				.VAlign(VAlign_Center)
 				[
@@ -84,7 +85,7 @@ public:
 					.FillHeight(1.0f)
 					[
 						SNew(STextBlock)
-						.WrapTextAt(350.0f)
+						.AutoWrapText(true)
 						.Text(Category.Description)
 						.TextStyle(&FEditorStyle::Get().GetWidgetStyle<FTextBlockStyle>("Tutorials.Browser.SummaryText"))
 						.HighlightText(HighlightText)
@@ -205,8 +206,7 @@ public:
 					]
 					+SHorizontalBox::Slot()
 					.FillWidth(1.0f)
-					.MaxWidth(400.0f)
-					.HAlign(HAlign_Left)
+					.HAlign(HAlign_Fill)
 					.VAlign(VAlign_Center)
 					[
 						SNew(SVerticalBox)
@@ -251,10 +251,9 @@ public:
 							]
 						]
 						+SVerticalBox::Slot()
-						.HAlign(HAlign_Left)
 						.FillHeight(1.0f)
 						[
-							STutorialContent::GenerateContentWidget(Tutorial->SummaryContent, 350.0f, false, DocumentationPage)
+							STutorialContent::GenerateContentWidget(Tutorial->SummaryContent, 500.0f, false, DocumentationPage)
 						]
 					]
 				]
@@ -387,8 +386,8 @@ void STutorialsBrowser::Construct(const FArguments& InArgs)
 				.AutoHeight()
 				[
 					SNew(SBox)
-					.HeightOverride(400.0f)
-					.WidthOverride(400.0f)
+					.HeightOverride(600.0f)
+					.WidthOverride(600.0f)
 					[
 						SNew(SScrollBox)
 						+SScrollBox::Slot()
@@ -505,21 +504,22 @@ void STutorialsBrowser::RebuildTutorials(TSharedPtr<FTutorialListEntry_Category>
 	TArray<TSharedPtr<FTutorialListEntry_Tutorial>> Tutorials;
 
 	// rebuild tutorials
-	for(const auto& TutorialReference : GetDefault<UTutorialSettings>()->Tutorials)
-	{
-		TSubclassOf<UEditorTutorial> TutorialClass = LoadClass<UEditorTutorial>(NULL, *TutorialReference.AssetLongPathname, NULL, LOAD_None, NULL);
-		if(TutorialClass != nullptr)
-		{
-			Tutorials.Add(MakeShareable(new FTutorialListEntry_Tutorial(TutorialClass->GetDefaultObject<UEditorTutorial>(), FOnTutorialSelected::CreateSP(this, &STutorialsBrowser::OnTutorialSelected), TAttribute<FText>::Create(TAttribute<FText>::FGetter::CreateSP(this, &STutorialsBrowser::GetSearchText)))));
-		}
-	}
+	FAssetRegistryModule& AssetRegistry = FModuleManager::LoadModuleChecked<FAssetRegistryModule>(TEXT("AssetRegistry"));
+
+	FARFilter Filter;
+	Filter.ClassNames.Add(UBlueprint::StaticClass()->GetFName());
+	Filter.bRecursiveClasses = true;
+	Filter.TagsAndValues.Add(TEXT("ParentClass"), FString::Printf(TEXT("%s'%s'"), *UClass::StaticClass()->GetName(), *UEditorTutorial::StaticClass()->GetPathName()));
+
+	TArray<FAssetData> AssetData;
+	AssetRegistry.Get().GetAssets(Filter, AssetData);
 
-	for(const auto& TutorialReference : GetDefault<UEditorTutorialSettings>()->Tutorials)
+	for (const auto& TutorialAsset : AssetData)
 	{
-		TSubclassOf<UEditorTutorial> TutorialClass = LoadClass<UEditorTutorial>(NULL, *TutorialReference.AssetLongPathname, NULL, LOAD_None, NULL);
-		if(TutorialClass != nullptr)
+		UBlueprint* Blueprint = LoadObject<UBlueprint>(nullptr, *TutorialAsset.ObjectPath.ToString());
+		if (Blueprint && Blueprint->GeneratedClass)
 		{
-			Tutorials.Add(MakeShareable(new FTutorialListEntry_Tutorial(TutorialClass->GetDefaultObject<UEditorTutorial>(), FOnTutorialSelected::CreateSP(this, &STutorialsBrowser::OnTutorialSelected), TAttribute<FText>::Create(TAttribute<FText>::FGetter::CreateSP(this, &STutorialsBrowser::GetSearchText)))));
+			Tutorials.Add(MakeShareable(new FTutorialListEntry_Tutorial(Blueprint->GeneratedClass->GetDefaultObject<UEditorTutorial>(), FOnTutorialSelected::CreateSP(this, &STutorialsBrowser::OnTutorialSelected), TAttribute<FText>::Create(TAttribute<FText>::FGetter::CreateSP(this, &STutorialsBrowser::GetSearchText)))));
 		}
 	}
 
diff --git a/Engine/Source/Editor/IntroTutorials/Private/TutorialSettings.h b/Engine/Source/Editor/IntroTutorials/Private/TutorialSettings.h
--- a/Engine/Source/Editor/IntroTutorials/Private/TutorialSettings.h
+++ b/Engine/Source/Editor/IntroTutorials/Private/TutorialSettings.h
@@ -15,10 +15,6 @@ class UTutorialSettings : public UObject
 	UPROPERTY(Config, EditAnywhere, Category="Tutorials")
 	TArray<FTutorialCategory> Categories;
 
-	/** The tutorials available in the editor */
-	UPROPERTY(Config, EditAnywhere, Category="Tutorials", meta=(MetaClass="EditorTutorial"))
-	TArray<FStringClassReference> Tutorials;
-
 	/** Tutorial to start on project startup */
 	UPROPERTY(Config, EditAnywhere, Category="Tutorials", meta=(MetaClass="EditorTutorial"))
 	FStringClassReference StartupTutorial;
diff --git a/Engine/Source/Editor/IntroTutorials/Public/EditorTutorial.h b/Engine/Source/Editor/IntroTutorials/Public/EditorTutorial.h
--- a/Engine/Source/Editor/IntroTutorials/Public/EditorTutorial.h
+++ b/Engine/Source/Editor/IntroTutorials/Public/EditorTutorial.h
@@ -205,24 +205,35 @@ class INTROTUTORIALS_API UEditorTutorial : public UObject
 	/** A standalone tutorial displays no navigation buttons and each content widget has a close button */
 	UPROPERTY(EditAnywhere, Category="Tutorial")
 	bool bIsStandalone;
+	
+public:
+	/** Called when a tutorial stage is started */
+	void HandleTutorialStageStarted(FName StageName);
+
+	/** Called when a tutorial stage ends */
+	void HandleTutorialStageEnded(FName StageName);
+
+	/** Called each tick so the Blueprint can optionally complete or skip stages */
+	void HandleTickCurrentStage(FName StageName);
 
+private:
 	/** Event fired when a tutorial stage begins */
 	UFUNCTION(BlueprintImplementableEvent, Category="Tutorial")
 	void OnTutorialStageStarted(FName StageName) const;
 
 	/** Event fired when a tutorial stage ends */
 	UFUNCTION(BlueprintImplementableEvent, Category="Tutorial")
-	void OnTutorialStageCompleted(FName StageName) const;
+	void OnTutorialStageEnded(FName StageName) const;
 
 	/** Advance to the next stage of a tutorial */
 	UFUNCTION(BlueprintCallable, Category="Tutorial")
-	static void GotoNextTutorialStage();
+	static void GoToNextTutorialStage();
 
 	/** Advance to the previous stage of a tutorial */
 	UFUNCTION(BlueprintCallable, Category="Tutorial")
-	static void GotoPreviousTutorialStage();
+	static void GoToPreviousTutorialStage();
 
 	/** Begin a tutorial. Note that this will end the current tutorial that is in progress, if any */
 	UFUNCTION(BlueprintCallable, Category="Tutorial")
-	static void BeginTutorial(UEditorTutorial* TutorialToStart);
+	static void BeginTutorial(UEditorTutorial* TutorialToStart, bool bRestart);
 };
\ No newline at end of file
