CIS fix for unresolved external

diff --git a/Engine/Source/Editor/IntroTutorials/Private/STutorialRoot.cpp b/Engine/Source/Editor/IntroTutorials/Private/STutorialRoot.cpp
--- a/Engine/Source/Editor/IntroTutorials/Private/STutorialRoot.cpp
+++ b/Engine/Source/Editor/IntroTutorials/Private/STutorialRoot.cpp
@@ -65,11 +65,6 @@ void STutorialRoot::Tick(const FGeometry& AllottedGeometry, const double InCurre
 	{
 		MaybeAddOverlay(Window);
 	}
-
-	if (CurrentTutorial != nullptr)
-	{
-		CurrentTutorial->HandleTickCurrentStage(CurrentTutorial->Stages[CurrentTutorialStage].Name);
-	}
 }
 
 void STutorialRoot::SummonTutorialBrowser(TSharedRef<SWindow> InWindow, const FString& InFilter)
