Allow disabling UI of tutorials while keeping the module intact.

Arguably sweeping a problem under the rug for now, but NullWidget added by STutorialRoot is apparently causing problems on Linux, so ability to disable UI (per platform) is helpful in the short term.

diff --git a/Engine/Config/Linux/LinuxEditorGameAgnostic.ini b/Engine/Config/Linux/LinuxEditorGameAgnostic.ini
new file mode 100644
--- /dev/null
+++ b/Engine/Config/Linux/LinuxEditorGameAgnostic.ini
@@ -0,0 +1,3 @@
+[IntroTutorials]
+DisableAllTutorials=true
+ 
diff --git a/Engine/Source/Editor/IntroTutorials/Private/IntroTutorials.cpp b/Engine/Source/Editor/IntroTutorials/Private/IntroTutorials.cpp
--- a/Engine/Source/Editor/IntroTutorials/Private/IntroTutorials.cpp
+++ b/Engine/Source/Editor/IntroTutorials/Private/IntroTutorials.cpp
@@ -34,6 +34,7 @@ IMPLEMENT_MODULE(FIntroTutorials, IntroTutorials)
 const FName FIntroTutorials::IntroTutorialTabName(TEXT("IntroTutorial"));
 
 const FString FIntroTutorials::IntroTutorialConfigSection(TEXT("IntroTutorials"));
+const FString FIntroTutorials::DisableTutorialsSettingName(TEXT("DisableAllTutorials"));
 const FString FIntroTutorials::InEditorTutorialPath(TEXT("Shared/Tutorials/inEditorTutorial"));
 const FString FIntroTutorials::WelcomeTutorialPath(TEXT("Shared/Tutorials/UE4Welcome"));
 const FString FIntroTutorials::InEditorGamifiedTutorialPath(TEXT("Shared/Tutorials/inEditorGamifiedTutorial"));
@@ -54,6 +55,7 @@ const FWelcomeTutorialProperties FIntroTutorials::TemplateOverview(TemplateOverv
 FIntroTutorials::FIntroTutorials()
 	: CurrentObjectClass(nullptr)
 {
+	bDisableTutorials = false;
 	bEnablePostTutorialSurveys = false;
 
 	bDesireResettingTutorialSeenFlagOnLoad = FParse::Param(FCommandLine::Get(), TEXT("tutorials"));
@@ -99,8 +101,10 @@ TSharedRef<FExtender> FIntroTutorials::AddSummonBlueprintTutorialsMenuExtender(c
 
 void FIntroTutorials::StartupModule()
 {
+	GConfig->GetBool(*IntroTutorialConfigSection, *DisableTutorialsSettingName, bDisableTutorials, GEditorGameAgnosticIni);
+
 	// This code can run with content commandlets. Slate is not initialized with commandlets and the below code will fail.
-	if (!IsRunningCommandlet())
+	if (!bDisableTutorials && !IsRunningCommandlet())
 	{
 		// Add tutorial for main frame opening
 		IMainFrameModule& MainFrameModule = FModuleManager::LoadModuleChecked<IMainFrameModule>(TEXT("MainFrame"));
@@ -223,9 +227,12 @@ void FIntroTutorials::StartupModule()
 
 void FIntroTutorials::ShutdownModule()
 {
-	FSourceCodeNavigation::AccessOnCompilerNotFound().RemoveAll( this );
+	if (!bDisableTutorials && !IsRunningCommandlet())
+	{
+		FSourceCodeNavigation::AccessOnCompilerNotFound().RemoveAll( this );
 
-	GLevelEditorModeTools().OnEditorModeChanged().RemoveAll(this);
+		GLevelEditorModeTools().OnEditorModeChanged().RemoveAll(this);
+	}
 
 	if (BlueprintEditorExtender.IsValid() && FModuleManager::Get().IsModuleLoaded("Kismet"))
 	{
@@ -661,6 +668,8 @@ void FIntroTutorials::HandleCompilerNotFound()
 	SummonTutorialWindowForPage( TEXT( "Shared/Tutorials/InstallingVisualStudioTutorial" ) );
 #elif PLATFORM_MAC
 	SummonTutorialWindowForPage( TEXT( "Shared/Tutorials/InstallingXCodeTutorial" ) );
+#else
+	STUBBED("FIntroTutorials::HandleCompilerNotFound");
 #endif
 }
 
diff --git a/Engine/Source/Editor/IntroTutorials/Private/IntroTutorials.h b/Engine/Source/Editor/IntroTutorials/Private/IntroTutorials.h
--- a/Engine/Source/Editor/IntroTutorials/Private/IntroTutorials.h
+++ b/Engine/Source/Editor/IntroTutorials/Private/IntroTutorials.h
@@ -45,6 +45,7 @@ class FIntroTutorials : public IIntroTutorials
 {
 public:
 	static const FString IntroTutorialConfigSection;
+	static const FString DisableTutorialsSettingName;
 	static const FString InEditorTutorialPath;
 	static const FString WelcomeTutorialPath;
 	static const FString InEditorGamifiedTutorialPath;
@@ -169,6 +170,8 @@ private:
 	/** The root window that we will always try to parent the tutorial window to */
 	TWeakPtr<SWindow> RootWindow;
 
+	/** Whether tutorials are disabled altogether */
+	bool bDisableTutorials;
 
 	TMap<FString, FGuid> TutorialSurveyMap;
 
