Possible fix for not being able to click on anything in the editor.  Possibly caused by mutation of SNullWidget visibility

diff --git a/Engine/Source/Editor/IntroTutorials/Private/SEditorTutorials.cpp b/Engine/Source/Editor/IntroTutorials/Private/SEditorTutorials.cpp
--- a/Engine/Source/Editor/IntroTutorials/Private/SEditorTutorials.cpp
+++ b/Engine/Source/Editor/IntroTutorials/Private/SEditorTutorials.cpp
@@ -60,7 +60,7 @@ void SEditorTutorials::Construct(const FArguments& InArgs)
 			+SVerticalBox::Slot()
 			.FillHeight(0.75f)
 			[
-				SNullWidget::NullWidget
+				SNew( SSpacer )
 			]
 			+SVerticalBox::Slot()
 			.FillHeight(0.25f)
