Fix a couple of windows being added as children when that won't work properly on OS X.

diff --git a/Engine/Source/Editor/IntroTutorials/Private/EditorTutorialFactory.cpp b/Engine/Source/Editor/IntroTutorials/Private/EditorTutorialFactory.cpp
--- a/Engine/Source/Editor/IntroTutorials/Private/EditorTutorialFactory.cpp
+++ b/Engine/Source/Editor/IntroTutorials/Private/EditorTutorialFactory.cpp
@@ -19,23 +19,20 @@ UEditorTutorialFactory::UEditorTutorialFactory(const class FPostConstructInitial
 
 bool UEditorTutorialFactory::ConfigureProperties()
 {
-	TSharedPtr<SWindow> RootWindow = FGlobalTabmanager::Get()->GetRootWindow();
-	if(RootWindow.IsValid())
-	{
-		// show import dialog
-		TSharedRef<SWindow> Window = SNew(SWindow)
-				.Title(LOCTEXT("WindowTitle", "Import Tutorial"))
-				.SupportsMaximize(false)
-				.SupportsMinimize(false)
-				.SizingRule(ESizingRule::Autosized);	
+	// show import dialog
+	TSharedRef<SWindow> Window = SNew(SWindow)
+			.Title(LOCTEXT("WindowTitle", "Import Tutorial"))
+			.SupportsMaximize(false)
+			.SupportsMinimize(false)
+			.SizingRule(ESizingRule::Autosized);	
 
-		ImporterDialog = SNew(SEditorTutorialImporter)
-			.ParentWindow(Window);
+	ImporterDialog = SNew(SEditorTutorialImporter)
+		.ParentWindow(Window);
 
-		Window->SetContent(ImporterDialog.ToSharedRef());
+	Window->SetContent(ImporterDialog.ToSharedRef());
 
-		FSlateApplication::Get().AddModalWindow(Window, RootWindow.ToSharedRef());
-	}
+	TSharedPtr<SWindow> RootWindow = FGlobalTabmanager::Get()->GetRootWindow();
+	FSlateApplication::Get().AddModalWindow(Window, RootWindow);
 
 	return true;
 }
diff --git a/Engine/Source/Editor/IntroTutorials/Private/IntroTutorials.cpp b/Engine/Source/Editor/IntroTutorials/Private/IntroTutorials.cpp
--- a/Engine/Source/Editor/IntroTutorials/Private/IntroTutorials.cpp
+++ b/Engine/Source/Editor/IntroTutorials/Private/IntroTutorials.cpp
@@ -379,10 +379,14 @@ void FIntroTutorials::SummonTutorialWindowForPage(const FString& Path)
 
 		TutorialWindow = Window;
 
-		if(RootWindow.IsValid())
+		if(RootWindow.IsValid() && !PLATFORM_MAC)
 		{
 			FSlateApplication::Get().AddWindowAsNativeChild(Window.ToSharedRef(), RootWindow.Pin().ToSharedRef());
 		}
+		else
+		{
+			FSlateApplication::Get().AddWindow(Window.ToSharedRef());
+		}
 
 		// tutorial Widget
 		TSharedRef<SIntroTutorials> TutorialWidgetRef = 
