Revised several graphnode tags for tutorial widgets.
Added tutorial widget support for autofocus (for BP nodes

diff --git a/Engine/Source/Editor/GraphEditor/Private/KismetNodes/SGraphNodeK2Base.cpp b/Engine/Source/Editor/GraphEditor/Private/KismetNodes/SGraphNodeK2Base.cpp
--- a/Engine/Source/Editor/GraphEditor/Private/KismetNodes/SGraphNodeK2Base.cpp
+++ b/Engine/Source/Editor/GraphEditor/Private/KismetNodes/SGraphNodeK2Base.cpp
@@ -48,8 +48,14 @@ void SGraphNodeK2Base::UpdateCompactNode()
 	FString TagName;
 	if (GraphNode != nullptr)
 	{
-		UBlueprint* Blueprint = FBlueprintEditorUtils::FindBlueprintForGraphChecked(CastChecked<UEdGraph>(GraphNode->GetOuter()));
-		TagName = FString::Printf(TEXT("GraphNode,%s,%s"), *GraphNode->GetName(), *GraphNode->NodeGuid.ToString());		
+		// We want the name of the blueprint as our name - we can find the node from the GUID
+		UObject* Package = GraphNode->GetOutermost();
+		UObject* LastOuter = GraphNode->GetOuter();
+		while (LastOuter->GetOuter() != Package)
+		{
+			LastOuter = LastOuter->GetOuter();
+		}
+		TagName = FString::Printf(TEXT("GraphNode,%s,%s"), *LastOuter->GetFullName(), *GraphNode->NodeGuid.ToString());
 	}
 	
 	//
diff --git a/Engine/Source/Editor/GraphEditor/Private/KismetNodes/SGraphNodeK2Var.cpp b/Engine/Source/Editor/GraphEditor/Private/KismetNodes/SGraphNodeK2Var.cpp
--- a/Engine/Source/Editor/GraphEditor/Private/KismetNodes/SGraphNodeK2Var.cpp
+++ b/Engine/Source/Editor/GraphEditor/Private/KismetNodes/SGraphNodeK2Var.cpp
@@ -163,24 +163,14 @@ void SGraphNodeK2Var::UpdateGraphNode()
 	FString TagName;
 	if (GraphNode != nullptr)
 	{
-		FString TagLabel;		
-		// Try to get a suitable name for the tag.
-		if (UK2Node_Variable const* K2Node = Cast<UK2Node_Variable>(GraphNode))
+		// We want the name of the blueprint as our name - we can find the node from the GUID
+		UObject* Package = GraphNode->GetOutermost();
+		UObject* LastOuter = GraphNode->GetOuter();
+		while (LastOuter->GetOuter() != Package)
 		{
-			UBlueprint* Blueprint = FBlueprintEditorUtils::FindBlueprintForGraphChecked(CastChecked<UEdGraph>(GraphNode->GetOuter()));
-			TagLabel = K2Node->GetVarNameString();
+			LastOuter = LastOuter->GetOuter();
 		}
-		else if (UK2Node_Self const* K2Node = Cast<UK2Node_Self>(GraphNode))
-		{
-			UBlueprint* Blueprint = FBlueprintEditorUtils::FindBlueprintForGraphChecked(CastChecked<UEdGraph>(GraphNode->GetOuter()));
-			TagLabel = GraphNode->NodeGuid.ToString();
-		}
-		else
-		{
-			UBlueprint* Blueprint = FBlueprintEditorUtils::FindBlueprintForGraphChecked(CastChecked<UEdGraph>(GraphNode->GetOuter()));
-			TagLabel = GraphNode->GetName();
-		}
-		TagName = FString::Printf(TEXT("GraphNode,%s,%s"), *TagLabel, *GraphNode->NodeGuid.ToString());
+		TagName = FString::Printf(TEXT("GraphNode,%s,%s"), *LastOuter->GetFullName(), *GraphNode->NodeGuid.ToString());
 	}
 
 	//             ________________
diff --git a/Engine/Source/Editor/GraphEditor/Private/SGraphNode.cpp b/Engine/Source/Editor/GraphEditor/Private/SGraphNode.cpp
--- a/Engine/Source/Editor/GraphEditor/Private/SGraphNode.cpp
+++ b/Engine/Source/Editor/GraphEditor/Private/SGraphNode.cpp
@@ -702,20 +702,16 @@ void SGraphNode::UpdateGraphNode()
 
 	// Setup a tag for this node
  	FString TagName;
- 	if (GraphNode != nullptr)
- 	{
-		UBlueprint* Blueprint = FBlueprintEditorUtils::FindBlueprintForGraph(Cast<UEdGraph>(GraphNode->GetOuter()));
-		FString NodeName;
-		if (Blueprint != nullptr)
+	if (GraphNode != nullptr)
+ 	{	
+		// We want the name of the blueprint as our name - we can find the node from the GUID
+		UObject* Package = GraphNode->GetOutermost();
+		UObject* LastOuter = GraphNode->GetOuter();
+		while (LastOuter->GetOuter() != Package)
 		{
-			NodeName = Blueprint->GetOuter()->GetName();
+			LastOuter = LastOuter->GetOuter();
 		}
-		else
-		{
-			NodeName = NodeTitle->GetHeadTitle().ToString();
-		}
-
-		TagName = FString::Printf(TEXT("GraphNode,%s,%s"), *NodeName, *GraphNode->NodeGuid.ToString());
+		TagName = FString::Printf(TEXT("GraphNode,%s,%s"), *LastOuter->GetFullName(), *GraphNode->NodeGuid.ToString());
  	}
 	
 	TSharedPtr<SVerticalBox> InnerVerticalBox;
diff --git a/Engine/Source/Editor/GraphEditor/Private/SGraphNodeComment.cpp b/Engine/Source/Editor/GraphEditor/Private/SGraphNodeComment.cpp
--- a/Engine/Source/Editor/GraphEditor/Private/SGraphNodeComment.cpp
+++ b/Engine/Source/Editor/GraphEditor/Private/SGraphNodeComment.cpp
@@ -102,8 +102,14 @@ void SGraphNodeComment::UpdateGraphNode()
 	FString TagName;
 	if (GraphNode != nullptr)
 	{
-		UBlueprint* Blueprint = FBlueprintEditorUtils::FindBlueprintForGraph(CastChecked<UEdGraph>(GraphNode->GetOuter()));
-		TagName = FString::Printf(TEXT("GraphNode,%s,%s"), *GetEditableNodeTitleAsText().ToString(), *GraphNode->NodeGuid.ToString());
+		// We want the name of the blueprint as our name - we can find the node from the GUID
+		UObject* Package = GraphNode->GetOutermost();
+		UObject* LastOuter = GraphNode->GetOuter();
+		while (LastOuter->GetOuter() != Package)
+		{
+			LastOuter = LastOuter->GetOuter();
+		}
+		TagName = FString::Printf(TEXT("GraphNode,%s,%s"), *LastOuter->GetFullName(), *GraphNode->NodeGuid.ToString());
 	}
 
 	TSharedPtr<SWidget> ErrorText = SetupErrorReporting();
diff --git a/Engine/Source/Editor/IntroTutorials/Private/STutorialOverlay.cpp b/Engine/Source/Editor/IntroTutorials/Private/STutorialOverlay.cpp
--- a/Engine/Source/Editor/IntroTutorials/Private/STutorialOverlay.cpp
+++ b/Engine/Source/Editor/IntroTutorials/Private/STutorialOverlay.cpp
@@ -6,6 +6,8 @@
 #include "EditorTutorial.h"
 #include "IntroTutorials.h"
 #include "LevelEditor.h"
+#include "BlueprintEditorUtils.h"
+#include "Guid.h"
 
 void STutorialOverlay::Construct(const FArguments& InArgs, FTutorialStage* const InStage)
 {
@@ -52,35 +54,33 @@ void STutorialOverlay::Construct(const FArguments& InArgs, FTutorialStage* const
 			FIntroTutorials& IntroTutorials = FModuleManager::Get().GetModuleChecked<FIntroTutorials>("IntroTutorials");
 
 			// now add canvas slots for widget-bound content
-			for(const FTutorialWidgetContent& WidgetContent : InStage->WidgetContent)
+			for (const FTutorialWidgetContent& WidgetContent : InStage->WidgetContent)
 			{
-				if(WidgetContent.Content.Type != ETutorialContent::None )
+				if (WidgetContent.Content.Type != ETutorialContent::None)
 				{
-					const bool bEmptyText = (WidgetContent.Content.Type == ETutorialContent::Text || WidgetContent.Content.Type == ETutorialContent::RichText) && WidgetContent.Content.Text.IsEmpty();
-					if(!bEmptyText)
-					{
-						TSharedPtr<STutorialContent> ContentWidget = 
-							SNew(STutorialContent, WidgetContent.Content)
-							.HAlign(WidgetContent.HorizontalAlignment)
-							.VAlign(WidgetContent.VerticalAlignment)
-							.Offset(WidgetContent.Offset)
-							.IsStandalone(bIsStandalone)
-							.OnClosed(OnClosed)
-							.WrapTextAt(WidgetContent.ContentWidth)
-							.Anchor(WidgetContent.WidgetAnchor);
-						OpenBrowserForWidgetAnchor(WidgetContent);
-
-						OverlayCanvas->AddSlot()
+
+					TSharedPtr<STutorialContent> ContentWidget =
+						SNew(STutorialContent, WidgetContent.Content)
+						.HAlign(WidgetContent.HorizontalAlignment)
+						.VAlign(WidgetContent.VerticalAlignment)
+						.Offset(WidgetContent.Offset)
+						.IsStandalone(bIsStandalone)
+						.OnClosed(OnClosed)
+						.WrapTextAt(WidgetContent.ContentWidth)
+						.Anchor(WidgetContent.WidgetAnchor);
+					
+					PerformWidgetInteractions(WidgetContent); 					
+
+					OverlayCanvas->AddSlot()
 						.Position(TAttribute<FVector2D>::Create(TAttribute<FVector2D>::FGetter::CreateSP(ContentWidget.Get(), &STutorialContent::GetPosition)))
 						.Size(TAttribute<FVector2D>::Create(TAttribute<FVector2D>::FGetter::CreateSP(ContentWidget.Get(), &STutorialContent::GetSize)))
 						[
 							ContentWidget.ToSharedRef()
 						];
 
-						OnPaintNamedWidget.AddSP(ContentWidget.Get(), &STutorialContent::HandlePaintNamedWidget);
-						OnResetNamedWidget.AddSP(ContentWidget.Get(), &STutorialContent::HandleResetNamedWidget);
-						OnCacheWindowSize.AddSP(ContentWidget.Get(), &STutorialContent::HandleCacheWindowSize);
-					}
+					OnPaintNamedWidget.AddSP(ContentWidget.Get(), &STutorialContent::HandlePaintNamedWidget);
+					OnResetNamedWidget.AddSP(ContentWidget.Get(), &STutorialContent::HandleResetNamedWidget);
+					OnCacheWindowSize.AddSP(ContentWidget.Get(), &STutorialContent::HandleCacheWindowSize);
 				}
 			}
 		}
@@ -141,12 +141,20 @@ int32 STutorialOverlay::TraverseWidgets(TSharedRef<SWidget> InWidget, const FGeo
 	return LayerId;
 }
 
+void STutorialOverlay::PerformWidgetInteractions(const FTutorialWidgetContent &WidgetContent)
+{
+	// Open any browser we need too
+	OpenBrowserForWidgetAnchor(WidgetContent);
+
+	FocusOnAnyBlueprintNodes(WidgetContent);
+}
+
 void STutorialOverlay::OpenBrowserForWidgetAnchor(const FTutorialWidgetContent &WidgetContent)
 {
 	FString IdentString = WidgetContent.WidgetAnchor.WrapperIdentifier.ToString();
 	FString TabString = "";
-		
-	// See if we have a mapping for the widget ident
+	
+	// See if we have a mapping for the widget ident.We will likely pull this from meta data when we have it)
 	for (TMap<FString, FString>::TConstIterator It(BrowserTabMap); It; ++It)
 	{
 		if (It.Key().StartsWith(IdentString) == true)
@@ -165,6 +173,40 @@ void STutorialOverlay::OpenBrowserForWidgetAnchor(const FTutorialWidgetContent &
 	
 }
 
+void STutorialOverlay::FocusOnAnyBlueprintNodes(const FTutorialWidgetContent &WidgetContent)
+{
+	if (WidgetContent.bAutoFocus == false)
+	{
+		return;
+	}
+	FString IdentString = WidgetContent.WidgetAnchor.WrapperIdentifier.ToString();
+	TArray<FString> Tokens;
+	IdentString.ParseIntoArray(&Tokens, TEXT(","), true);
+	if (Tokens.Num() > 1 )
+	{
+		FString Name = Tokens[1];
+
+		// I'm sure this isnt how to do this. But probably this will all come from meta data anyway
+		int32 NameIndex;
+		Name.FindLastChar(TEXT('.'), NameIndex);
+		FString BlueprintName = Name.RightChop(NameIndex+1);		
+		UBlueprint* Blueprint = FindObject<UBlueprint>(ANY_PACKAGE, *BlueprintName);
+		// If we find a blueprint
+		if (Blueprint != nullptr)
+		{
+			// Try to grab guid
+			FGuid NodeGuid;
+			FGuid::Parse(Tokens[Tokens.Num() - 1], NodeGuid);
+			UEdGraphNode* OutNode = NULL;
+			if (UEdGraphNode* GraphNode = FBlueprintEditorUtils::GetNodeByGUID(Blueprint, NodeGuid))
+			{
+				FKismetEditorUtilities::BringKismetToFocusAttentionOnObject(GraphNode, false);
+			}
+		}
+
+	}
+}
+
 void STutorialOverlay::AddTabInfo()
 {
 	BrowserTabMap.Empty();
diff --git a/Engine/Source/Editor/IntroTutorials/Private/STutorialOverlay.h b/Engine/Source/Editor/IntroTutorials/Private/STutorialOverlay.h
--- a/Engine/Source/Editor/IntroTutorials/Private/STutorialOverlay.h
+++ b/Engine/Source/Editor/IntroTutorials/Private/STutorialOverlay.h
@@ -49,6 +49,12 @@ private:
 	/* Opens the browser that the given widget requires if it is not already. */
 	void OpenBrowserForWidgetAnchor(const FTutorialWidgetContent &WidgetContent);
 
+	/* Focus on the blueprint node if an anchor widget references one */
+	void FocusOnAnyBlueprintNodes(const FTutorialWidgetContent &WidgetContent);
+
+	/* Do any interaction stuff for a widget - open browser, scroll to node etc */
+	void PerformWidgetInteractions(const FTutorialWidgetContent &WidgetContent);
+
 	/* Creates a map of broswers required by the various widget types. */
 	void AddTabInfo();
 
diff --git a/Engine/Source/Editor/IntroTutorials/Public/EditorTutorial.h b/Engine/Source/Editor/IntroTutorials/Public/EditorTutorial.h
--- a/Engine/Source/Editor/IntroTutorials/Public/EditorTutorial.h
+++ b/Engine/Source/Editor/IntroTutorials/Public/EditorTutorial.h
@@ -131,6 +131,7 @@ struct INTROTUTORIALS_API FTutorialWidgetContent
 		ContentWidth = 350.0f;
 		HorizontalAlignment = HAlign_Center;
 		VerticalAlignment = VAlign_Bottom;
+		bAutoFocus = false;
 	}
 
 	/** Content to associate with widget */
@@ -154,6 +155,10 @@ struct INTROTUTORIALS_API FTutorialWidgetContent
 	/** Content width - text will be wrapped at this point */
 	UPROPERTY(EditAnywhere, AdvancedDisplay, Category="Widget", meta=(UIMin="10.0", UIMax="600.0"))
 	float ContentWidth;
+
+	/** If this a node that can be focused (EG a blueprint node) should we auto focus on it */
+	UPROPERTY(EditAnywhere, Category = "Anchor")
+	bool bAutoFocus;
 };
 
 /** A single tutorial stage, containing the optional main content & a number of widgets with content attached */
diff --git a/Engine/Source/Editor/Kismet/Private/SBlueprintPalette.cpp b/Engine/Source/Editor/Kismet/Private/SBlueprintPalette.cpp
--- a/Engine/Source/Editor/Kismet/Private/SBlueprintPalette.cpp
+++ b/Engine/Source/Editor/Kismet/Private/SBlueprintPalette.cpp
@@ -879,6 +879,12 @@ void SBlueprintPaletteItem::Construct(const FArguments& InArgs, FCreateWidgetFor
 	GetPaletteItemIcon(GraphAction, InBlueprintEditor.Pin()->GetBlueprintObj(), IconBrush, IconColor, IconToolTip, IconDocLink, IconDocExcerpt);
 	TSharedRef<SWidget> IconWidget = CreateIconWidget(IconToolTip, IconBrush, IconColor, IconDocLink, IconDocExcerpt);
 
+	// Setup a tag for this node
+	FString TagName;
+	if( ActionPtr.IsValid() )
+	{
+		TagName = FString::Printf(TEXT("PaletteItem,%s,%d"), *GraphAction->MenuDescription.ToString(), GraphAction->SectionID);
+	}
 	// construct the text widget
 	FSlateFontInfo NameFont = FSlateFontInfo(FPaths::EngineContentDir() / TEXT("Slate/Fonts/Roboto-Regular.ttf"), 10);
 	bool bIsReadOnly = FBlueprintEditorUtils::IsPaletteActionReadOnly(GraphAction, InBlueprintEditor.Pin());
@@ -888,7 +894,7 @@ void SBlueprintPaletteItem::Construct(const FArguments& InArgs, FCreateWidgetFor
 	ChildSlot
 	[
 		SNew(SHorizontalBox)
-			
+		.Tag(*TagName)
 		// icon slot
 		+SHorizontalBox::Slot()
 			.AutoWidth()
