non unity fix

diff --git a/Engine/Source/Editor/IntroTutorials/Private/STutorialOverlay.cpp b/Engine/Source/Editor/IntroTutorials/Private/STutorialOverlay.cpp
--- a/Engine/Source/Editor/IntroTutorials/Private/STutorialOverlay.cpp
+++ b/Engine/Source/Editor/IntroTutorials/Private/STutorialOverlay.cpp
@@ -5,6 +5,7 @@
 #include "STutorialContent.h"
 #include "EditorTutorial.h"
 #include "IntroTutorials.h"
+#include "Kismet2/KismetEditorUtilities.h"
 #include "LevelEditor.h"
 #include "BlueprintEditorUtils.h"
 #include "Guid.h"
