Temporary build fix.

diff --git a/Engine/Source/Editor/IntroTutorials/Private/STutorialContent.cpp b/Engine/Source/Editor/IntroTutorials/Private/STutorialContent.cpp
--- a/Engine/Source/Editor/IntroTutorials/Private/STutorialContent.cpp
+++ b/Engine/Source/Editor/IntroTutorials/Private/STutorialContent.cpp
@@ -103,7 +103,8 @@ static void GetAnimationValues(bool bIsIntro, float InAnimationProgress, float&
 int32 STutorialContent::OnPaint( const FPaintArgs& Args, const FGeometry& AllottedGeometry, const FSlateRect& MyClippingRect, FSlateWindowElementList& OutDrawElements, int32 LayerId, const FWidgetStyle& InWidgetStyle, bool bParentEnabled ) const
 {
 	CachedContentGeometry = AllottedGeometry;
-	CachedContentGeometry.AbsolutePosition += OutDrawElements.GetWindow()->GetPositionInScreen();
+	// BM: Fixing compile error. TomS to fix.
+	// CachedContentGeometry.AbsolutePosition += OutDrawElements.GetWindow()->GetPositionInScreen();
 
 	if(bIsVisible && Anchor.Type != ETutorialAnchorIdentifier::None && Anchor.bDrawHighlight)
 	{
