Simplified rich text editor for tutorials.

Removed ability to independently set font size & colour etc. and replaced with a set of font styles we can use (and style appropriately). Added legacy code to deal with old-style text as best it can.

diff --git a/Engine/Source/Editor/Documentation/Private/UDNParser.cpp b/Engine/Source/Editor/Documentation/Private/UDNParser.cpp
--- a/Engine/Source/Editor/Documentation/Private/UDNParser.cpp
+++ b/Engine/Source/Editor/Documentation/Private/UDNParser.cpp
@@ -724,7 +724,7 @@ void FUDNParser::AddContentToExcerpt(TSharedPtr<SVerticalBox> Box, const FString
 		);
 
 		AddLineSeperator(Excerpt);
-		Excerpt.RichText += FString::Printf(TEXT("<TextStyle FontFamily=\"Roboto\" FontSize=\"11\" FontStyle=\"Regular\" FontColor=\"(R=0.0,G=0.0,B=0.0,A=1.0)\">%s</>"), *ContentSource);
+		Excerpt.RichText += FString::Printf(TEXT("<TextStyle Style=\"%s\">%s</>"), *Style.ContentStyleName.ToString(), *ContentSource);
 	}
 }
 
@@ -860,7 +860,7 @@ TSharedRef< SWidget > FUDNParser::GenerateExcerptContent( const FString& Link, F
 					);
 
 				AddLineSeperator(Excerpt);
-				Excerpt.RichText += FString::Printf(TEXT("<TextStyle FontFamily=\"Roboto\" FontSize=\"11\" FontStyle=\"Bold\" FontColor=\"(R=0.0,G=0.0,B=0.0,A=1.0)\">%s</>"), *Line.AdditionalContent[0]);
+				Excerpt.RichText += FString::Printf(TEXT("<TextStyle Style=\"%s\">%s</>"), *Style.BoldContentStyleName.ToString(), *Line.AdditionalContent[0]);
 				break;
 			case FUDNLine::NumberedContent:
 				AddContentToExcerpt(Box, CurrentStringContent, Excerpt);
@@ -887,7 +887,6 @@ TSharedRef< SWidget > FUDNParser::GenerateExcerptContent( const FString& Link, F
 					];
 
 				AddLineSeperator(Excerpt);
-				Excerpt.RichText += TEXT("<hr></>");
 				break;
 			case FUDNLine::Header1:
 				AddContentToExcerpt(Box, CurrentStringContent, Excerpt);
@@ -900,7 +899,7 @@ TSharedRef< SWidget > FUDNParser::GenerateExcerptContent( const FString& Link, F
 					);
 
 				AddLineSeperator(Excerpt);
-				Excerpt.RichText += FString::Printf(TEXT("<TextStyle FontFamily=\"Roboto\" FontSize=\"32\" FontStyle=\"Bold\" FontColor=\"(R=0.0,G=0.0,B=0.0,A=1.0)\">%s</>"), *Line.AdditionalContent[0]);
+				Excerpt.RichText += FString::Printf(TEXT("<TextStyle Style=\"%s\">%s</>"), *Style.Header1StyleName.ToString(), *Line.AdditionalContent[0]);
 				break;
 			case FUDNLine::Header2:
 				AddContentToExcerpt(Box, CurrentStringContent, Excerpt);
@@ -913,7 +912,7 @@ TSharedRef< SWidget > FUDNParser::GenerateExcerptContent( const FString& Link, F
 					);
 
 				AddLineSeperator(Excerpt);
-				Excerpt.RichText += FString::Printf(TEXT("<TextStyle FontFamily=\"Roboto\" FontSize=\"24\" FontStyle=\"Bold\" FontColor=\"(R=0.0,G=0.0,B=0.0,A=1.0)\">%s</>"), *Line.AdditionalContent[0]);
+				Excerpt.RichText += FString::Printf(TEXT("<TextStyle Style=\"%s\">%s</>"), *Style.Header2StyleName.ToString(), *Line.AdditionalContent[0]);
 				break;
 			case FUDNLine::Link:
 				AddContentToExcerpt(Box, CurrentStringContent, Excerpt);
diff --git a/Engine/Source/Editor/EditorStyle/Private/SlateEditorStyle.cpp b/Engine/Source/Editor/EditorStyle/Private/SlateEditorStyle.cpp
--- a/Engine/Source/Editor/EditorStyle/Private/SlateEditorStyle.cpp
+++ b/Engine/Source/Editor/EditorStyle/Private/SlateEditorStyle.cpp
@@ -2472,24 +2472,44 @@ void FSlateEditorStyle::FStyle::SetupTutorialStyles()
 		// Default text styles
 		{
 			const FTextBlockStyle RichTextNormal = FTextBlockStyle()
-				.SetFont(TTF_FONT("Fonts/Roboto-Regular", 11))
+				.SetFont(TTF_CORE_FONT("Fonts/Roboto-Regular", 11))
 				.SetColorAndOpacity(FSlateColor::UseForeground())
 				.SetShadowOffset(FVector2D::ZeroVector)
 				.SetShadowColorAndOpacity(FLinearColor::Black)
 				.SetHighlightColor(FLinearColor(0.02f, 0.3f, 0.0f))
 				.SetHighlightShape(BOX_BRUSH("Common/TextBlockHighlightShape", FMargin(3.f /8.f)));
-			Set( "TutorialEditableText.Editor.Text", RichTextNormal );
+			Set( "Tutorials.Content.Text", RichTextNormal );
+
+			Set( "Tutorials.Content.TextBold", FTextBlockStyle(RichTextNormal)
+				.SetFont(TTF_CORE_FONT("Fonts/Roboto-Bold", 11)));
+
+			Set( "Tutorials.Content.HeaderText1", FTextBlockStyle(RichTextNormal)
+				.SetFontSize(20));
+
+			Set( "Tutorials.Content.HeaderText2", FTextBlockStyle(RichTextNormal)
+				.SetFontSize(16));
 
 			{
 				const FButtonStyle RichTextHyperlinkButton = FButtonStyle()
-					.SetNormal(BORDER_BRUSH("Old/HyperlinkDotted", FMargin(0,0,0,3/16.0f), FSlateColor::UseForeground() ) )
+					.SetNormal(BORDER_BRUSH("Old/HyperlinkDotted", FMargin(0,0,0,3/16.0f), FLinearColor::Blue ) )
 					.SetPressed(FSlateNoResource() )
-					.SetHovered(BORDER_BRUSH("Old/HyperlinkUnderline", FMargin(0,0,0,3/16.0f), FSlateColor::UseForeground() ) );
+					.SetHovered(BORDER_BRUSH("Old/HyperlinkUnderline", FMargin(0,0,0,3/16.0f), FLinearColor::Blue ) );
+
+				const FTextBlockStyle RichTextHyperlinkText = FTextBlockStyle(RichTextNormal)
+					.SetColorAndOpacity(FLinearColor::Blue);
+
+				Set( "Tutorials.Content.HyperlinkText", RichTextHyperlinkText );
+
+				// legacy style
+				Set( "TutorialEditableText.Editor.HyperlinkText", RichTextHyperlinkText );
 
 				const FHyperlinkStyle RichTextHyperlink = FHyperlinkStyle()
 					.SetUnderlineStyle(RichTextHyperlinkButton)
-					.SetTextStyle(RichTextNormal)
+					.SetTextStyle(RichTextHyperlinkText)
 					.SetPadding(FMargin(0.0f));
+				Set( "Tutorials.Content.Hyperlink", RichTextHyperlink );
+
+				// legacy style
 				Set( "TutorialEditableText.Editor.Hyperlink", RichTextHyperlink );
 			}
 		}
@@ -2509,17 +2529,17 @@ void FSlateEditorStyle::FStyle::SetupTutorialStyles()
 			Set("TutorialEditableText.Toolbar.TextColor", TextColor);
 
 			Set("TutorialEditableText.Toolbar.Text", FTextBlockStyle(NormalText)
-				.SetFont(TTF_FONT("Fonts/Roboto-Regular", 10))
+				.SetFont(TTF_CORE_FONT("Fonts/Roboto-Regular", 10))
 				.SetColorAndOpacity(TextColor)
 				);
 
 			Set("TutorialEditableText.Toolbar.BoldText", FTextBlockStyle(NormalText)
-				.SetFont(TTF_FONT("Fonts/Roboto-Bold", 10))
+				.SetFont(TTF_CORE_FONT("Fonts/Roboto-Bold", 10))
 				.SetColorAndOpacity(TextColor)
 				);
 
 			Set("TutorialEditableText.Toolbar.ItalicText", FTextBlockStyle(NormalText)
-				.SetFont(TTF_FONT("Fonts/Roboto-Italic", 10))
+				.SetFont(TTF_CORE_FONT("Fonts/Roboto-Italic", 10))
 				.SetColorAndOpacity(TextColor)
 				);
 
diff --git a/Engine/Source/Editor/IntroTutorials/Private/STutorialContent.cpp b/Engine/Source/Editor/IntroTutorials/Private/STutorialContent.cpp
--- a/Engine/Source/Editor/IntroTutorials/Private/STutorialContent.cpp
+++ b/Engine/Source/Editor/IntroTutorials/Private/STutorialContent.cpp
@@ -173,13 +173,13 @@ TSharedRef<SWidget> STutorialContent::GenerateContentWidget(const FTutorialConte
 	// Style for the documentation
 	static FDocumentationStyle DocumentationStyle;
 	DocumentationStyle
-		.ContentStyle(TEXT("Tutorials.Content"))
-		.BoldContentStyle(TEXT("Tutorials.BoldContent"))
-		.NumberedContentStyle(TEXT("Tutorials.NumberedContent"))
-		.Header1Style(TEXT("Tutorials.Header1"))
-		.Header2Style(TEXT("Tutorials.Header2"))
-		.HyperlinkButtonStyle(TEXT("Tutorials.Hyperlink.Button"))
-		.HyperlinkTextStyle(TEXT("Tutorials.Hyperlink.Text"))
+		.ContentStyle(TEXT("Tutorials.Content.Text"))
+		.BoldContentStyle(TEXT("Tutorials.Content.TextBold"))
+		.NumberedContentStyle(TEXT("Tutorials.Content.Text"))
+		.Header1Style(TEXT("Tutorials.Content.Header1"))
+		.Header2Style(TEXT("Tutorials.Content.Header2"))
+		.HyperlinkButtonStyle(TEXT("Tutorials.Content.Hyperlink"))
+		.HyperlinkTextStyle(TEXT("Tutorials.Content.HyperlinkText"))
 		.SeparatorStyle(TEXT("Tutorials.Separator"));
 
 	OutDocumentationPage = nullptr;
@@ -222,27 +222,17 @@ TSharedRef<SWidget> STutorialContent::GenerateContentWidget(const FTutorialConte
 
 	case ETutorialContent::RichText:
 		{
-			FTextStyles TextStyles;
-			TextStyles.AvailableFontFamilies.Emplace(MakeShareable(new FTextStyles::FFontFamily(
-				LOCTEXT("RobotoFontFamily", "Roboto"), 
-				TEXT("Roboto"), 
-				FName(*(FPaths::EngineContentDir() / TEXT("Slate/Fonts/Roboto-Regular.ttf"))),
-				FName(*(FPaths::EngineContentDir() / TEXT("Slate/Fonts/Roboto-Bold.ttf"))),
-				FName(*(FPaths::EngineContentDir() / TEXT("Slate/Fonts/Roboto-Italic.ttf"))),
-				FName(*(FPaths::EngineContentDir() / TEXT("Slate/Fonts/Roboto-BoldItalic.ttf")))
-				)));
-
 			TArray< TSharedRef< class ITextDecorator > > Decorators;
 			Decorators.Add(FHyperlinkDecorator::Create(TEXT("browser"), FSlateHyperlinkRun::FOnClick::CreateStatic(&TutorialTextHelpers::OnBrowserLinkClicked)));
 			Decorators.Add(FHyperlinkDecorator::Create(TEXT("udn"), FSlateHyperlinkRun::FOnClick::CreateStatic(&TutorialTextHelpers::OnDocLinkClicked)));
 			Decorators.Add(FHyperlinkDecorator::Create(TEXT("tutorial"), FSlateHyperlinkRun::FOnClick::CreateStatic(&TutorialTextHelpers::OnTutorialLinkClicked)));
 			Decorators.Add(FHyperlinkDecorator::Create(TEXT("code"), FSlateHyperlinkRun::FOnClick::CreateStatic(&TutorialTextHelpers::OnCodeLinkClicked)));
 			Decorators.Add(FHyperlinkDecorator::Create(TEXT("asset"), FSlateHyperlinkRun::FOnClick::CreateStatic(&TutorialTextHelpers::OnAssetLinkClicked)));
-			Decorators.Add(FTextStyleDecorator::Create(TextStyles));
+			Decorators.Add(FTextStyleDecorator::Create());
 
 			return SNew(SRichTextBlock)
 					.Visibility(EVisibility::SelfHitTestInvisible)
-					.TextStyle(FEditorStyle::Get(), "TutorialEditableText.Editor.Text")
+					.TextStyle(FEditorStyle::Get(), "Tutorials.Content.Text")
 					.DecoratorStyleSet(&FEditorStyle::Get())
 					.Decorators(Decorators)
 					.Text(InContent.Text)
diff --git a/Engine/Source/Editor/IntroTutorials/Private/STutorialEditableText.cpp b/Engine/Source/Editor/IntroTutorials/Private/STutorialEditableText.cpp
--- a/Engine/Source/Editor/IntroTutorials/Private/STutorialEditableText.cpp
+++ b/Engine/Source/Editor/IntroTutorials/Private/STutorialEditableText.cpp
@@ -151,29 +151,23 @@ void STutorialEditableText::Construct(const FArguments& InArgs)
 	OnTextChanged = InArgs._OnTextChanged;
 	OnTextCommitted = InArgs._OnTextCommitted;
 
-	bPickingColor = false;
 	CurrentHyperlinkType = EHyperlinkType::Browser;
-
-	// Define and add the "Roboto" font family
-	TextStyles.AvailableFontFamilies.Emplace(MakeShareable(new FTextStyles::FFontFamily(
-		LOCTEXT("RobotoFontFamily", "Roboto"), 
-		TEXT("Roboto"), 
-		FName(*(FPaths::EngineContentDir() / TEXT("Slate/Fonts/Roboto-Regular.ttf"))),
-		FName(*(FPaths::EngineContentDir() / TEXT("Slate/Fonts/Roboto-Bold.ttf"))),
-		FName(*(FPaths::EngineContentDir() / TEXT("Slate/Fonts/Roboto-Italic.ttf"))),
-		FName(*(FPaths::EngineContentDir() / TEXT("Slate/Fonts/Roboto-BoldItalic.ttf")))
-		)));
-
-	// Set some sensible defaults (these also match the default text style of "TutorialEditableText.Editor.Text"
-	ActiveFontFamily = TextStyles.AvailableFontFamilies[0];
-	FontSize = 11;
-	FontStyle = FTextStyles::EFontStyle::Regular;
-	FontColor = FLinearColor::Black;
+	bNewHyperlink = true;
+
+	// Setup text styles
+	StylesAndNames.Add(MakeShareable(new FTextStyleAndName(TEXT("Tutorials.Content.Text"), LOCTEXT("NormalTextDesc", "Normal"))));
+	StylesAndNames.Add(MakeShareable(new FTextStyleAndName(TEXT("Tutorials.Content.TextBold"), LOCTEXT("BoldTextDesc", "Bold"))));
+	StylesAndNames.Add(MakeShareable(new FTextStyleAndName(TEXT("Tutorials.Content.HeaderText2"), LOCTEXT("Header2TextDesc", "Header 2"))));
+	StylesAndNames.Add(MakeShareable(new FTextStyleAndName(TEXT("Tutorials.Content.HeaderText1"), LOCTEXT("Header1TextDesc", "Header 1"))));
+	ActiveStyle = StylesAndNames[0];
+	
+	HyperlinkStyle = MakeShareable(new FTextStyleAndName(TEXT("Tutorials.Content.HyperlinkText"), LOCTEXT("HyperlinkTextDesc", "Hyperlink")));
+	StylesAndNames.Add(HyperlinkStyle);
 
 	TSharedRef<FRichTextLayoutMarshaller> RichTextMarshaller = FRichTextLayoutMarshaller::Create(
 		TArray<TSharedRef<ITextDecorator>>(), 
 		&FEditorStyle::Get(), 
-		FEditorStyle::Get().GetWidgetStyle<FTextBlockStyle>("TutorialEditableText.Editor.Text")
+		FEditorStyle::Get().GetWidgetStyle<FTextBlockStyle>("Tutorials.Content.Text")
 		);
 
 	OnBrowserLinkClicked = FSlateHyperlinkRun::FOnClick::CreateStatic(&TutorialTextHelpers::OnBrowserLinkClicked);
@@ -186,7 +180,7 @@ void STutorialEditableText::Construct(const FArguments& InArgs)
 	RichTextMarshaller->AppendInlineDecorator(FHyperlinkDecorator::Create(TEXT("tutorial"), OnTutorialLinkClicked));
 	RichTextMarshaller->AppendInlineDecorator(FHyperlinkDecorator::Create(TEXT("code"), OnCodeLinkClicked));
 	RichTextMarshaller->AppendInlineDecorator(FHyperlinkDecorator::Create(TEXT("asset"), OnAssetLinkClicked));
-	RichTextMarshaller->AppendInlineDecorator(FTextStyleDecorator::Create(TextStyles));
+	RichTextMarshaller->AppendInlineDecorator(FTextStyleDecorator::Create());
 
 	this->ChildSlot
 	[
@@ -197,7 +191,7 @@ void STutorialEditableText::Construct(const FArguments& InArgs)
 		.Padding(FMargin(0.0f, 0.0f, 0.0f, 0.0f))
 		[
 			SAssignNew(RichEditableTextBox, SMultiLineEditableTextBox)
-			.Font(FEditorStyle::Get().GetWidgetStyle<FTextBlockStyle>("TutorialEditableText.Editor.Text").Font)
+			.Font(FEditorStyle::Get().GetWidgetStyle<FTextBlockStyle>("Tutorials.Content.Text").Font)
 			.Text(InArgs._Text)
 			.OnTextChanged(this, &STutorialEditableText::HandleRichEditableTextChanged)
 			.OnTextCommitted(this, &STutorialEditableText::HandleRichEditableTextCommitted)
@@ -222,102 +216,20 @@ void STutorialEditableText::Construct(const FArguments& InArgs)
 				+SHorizontalBox::Slot()
 				.AutoWidth()
 				[
-					SAssignNew(FontComboBox, SComboBox<TSharedPtr<FTextStyles::FFontFamily>>)
+					SAssignNew(FontComboBox, SComboBox<TSharedPtr<FTextStyleAndName>>)
 					.ComboBoxStyle(FEditorStyle::Get(), "TutorialEditableText.Toolbar.ComboBox")
-					.OptionsSource(&TextStyles.AvailableFontFamilies)
-					.OnSelectionChanged(this, &STutorialEditableText::OnActiveFontFamilyChanged)
-					.OnGenerateWidget(this, &STutorialEditableText::GenerateFontFamilyComboEntry)
-					.InitiallySelectedItem(ActiveFontFamily)
+					.OptionsSource(&StylesAndNames)
+					.OnSelectionChanged(this, &STutorialEditableText::OnActiveStyleChanged)
+					.OnGenerateWidget(this, &STutorialEditableText::GenerateStyleComboEntry)
+					.ContentPadding(0.0f)
+					.InitiallySelectedItem(nullptr)
 					[
 						SNew(SBox)
 						.Padding(FMargin(0.0f, 0.0f, 2.0f, 0.0f))
+						.MinDesiredWidth(100.0f)
 						[
 							SNew(STextBlock)
-							.Text(this, &STutorialEditableText::GetActiveFontFamilyName)
-						]
-					]
-				]
-
-				+SHorizontalBox::Slot()
-				.AutoWidth()
-				[
-					SNew(SBox)
-					.WidthOverride(24)
-					[
-						SNew(SNumericEntryBox<uint8>)
-						.Value(this, &STutorialEditableText::GetFontSize)
-						.OnValueCommitted(this, &STutorialEditableText::SetFontSize)
-					]
-				]
-
-				+SHorizontalBox::Slot()
-				.AutoWidth()
-				.Padding(FMargin(4.0f, 0.0f, 0.0f, 0.0f))
-				[
-					SNew(SCheckBox)
-					.Style(FEditorStyle::Get(), "TutorialEditableText.Toolbar.ToggleButtonCheckbox")
-					.IsChecked(this, &STutorialEditableText::IsFontStyleBold)
-					.OnCheckStateChanged(this, &STutorialEditableText::OnFontStyleBoldChanged)
-					[
-						SNew(SBox)
-						.WidthOverride(24)
-						.HAlign(HAlign_Center)
-						.VAlign(VAlign_Center)
-						[
-							SNew(STextBlock)
-							.TextStyle(FEditorStyle::Get(), "TutorialEditableText.Toolbar.BoldText")
-							.Text(LOCTEXT("BoldLabel", "B"))
-						]
-					]
-				]
-
-				+SHorizontalBox::Slot()
-				.AutoWidth()
-				[
-					SNew(SCheckBox)
-					.Style(FEditorStyle::Get(), "TutorialEditableText.Toolbar.ToggleButtonCheckbox")
-					.IsChecked(this, &STutorialEditableText::IsFontStyleItalic)
-					.OnCheckStateChanged(this, &STutorialEditableText::OnFontStyleItalicChanged)
-					[
-						SNew(SBox)
-						.WidthOverride(24)
-						.HAlign(HAlign_Center)
-						.VAlign(VAlign_Center)
-						[
-							SNew(STextBlock)
-							.TextStyle(FEditorStyle::Get(), "TutorialEditableText.Toolbar.ItalicText")
-							.Text(LOCTEXT("ItalicLabel", "I"))
-						]
-					]
-				]
-
-				+SHorizontalBox::Slot()
-				.Padding(FMargin(4.0f, 0.0f, 0.0f, 0.0f))
-				.AutoWidth()
-				[
-					SNew(SButton)
-					.ButtonStyle(FEditorStyle::Get(), "TutorialEditableText.Toolbar.Button")
-					.OnClicked(this, &STutorialEditableText::OpenFontColorPicker)
-					[
-						SNew(SOverlay)
-
-						+SOverlay::Slot()
-						.Padding(FMargin(0.0f, 0.0f, 0.0f, 4.0f))
-						.HAlign(HAlign_Center)
-						.VAlign(VAlign_Bottom)
-						[
-							SNew(STextBlock)
-							.TextStyle(FEditorStyle::Get(), "TutorialEditableText.Toolbar.BoldText")
-							.Text(LOCTEXT("ColorLabel", "A"))
-						]
-
-						+SOverlay::Slot()
-						.HAlign(HAlign_Center)
-						.VAlign(VAlign_Bottom)
-						[
-							SNew(SColorBlock)
-							.Color(this, &STutorialEditableText::GetFontColor)
-							.Size(FVector2D(20.0f, 6.0f))
+							.Text(this, &STutorialEditableText::GetActiveStyleName)
 						]
 					]
 				]
@@ -327,19 +239,15 @@ void STutorialEditableText::Construct(const FArguments& InArgs)
 				.AutoWidth()
 				[
 					SAssignNew(HyperlinkComboButton, SComboButton)
+					.ToolTipText(LOCTEXT("HyperlinkButtonTooltip", "Insert or Edit Hyperlink"))
 					.ComboButtonStyle(FEditorStyle::Get(), "TutorialEditableText.Toolbar.ComboButton")
-					.HasDownArrow(false)
 					.OnComboBoxOpened(this, &STutorialEditableText::HandleHyperlinkComboOpened)
+					.IsEnabled(this, &STutorialEditableText::IsHyperlinkComboEnabled)
+					.ContentPadding(1.0f)
 					.ButtonContent()
 					[
-						SNew(SBox)
-						.WidthOverride(20)
-						.HAlign(HAlign_Center)
-						.VAlign(VAlign_Center)
-						[
-							SNew(SImage)
-							.Image(FEditorStyle::Get().GetBrush("TutorialEditableText.Toolbar.HyperlinkImage"))
-						]
+						SNew(SImage)
+						.Image(FEditorStyle::Get().GetBrush("TutorialEditableText.Toolbar.HyperlinkImage"))
 					]
 					.MenuContent()
 					[
@@ -360,7 +268,8 @@ void STutorialEditableText::Construct(const FArguments& InArgs)
 							SNew(SBox)
 							.WidthOverride(300)
 							[
-								SAssignNew(HyperlinkNameTextBox, SEditableTextBox)
+								SAssignNew(HyperlinkNameTextBlock, STextBlock)
+								.TextStyle(FEditorStyle::Get(), "TutorialEditableText.Toolbar.Text")
 							]
 						]
 
@@ -461,11 +370,11 @@ void STutorialEditableText::Construct(const FArguments& InArgs)
 							[
 								SNew(SButton)
 								.ButtonStyle(FEditorStyle::Get(), "TutorialEditableText.Toolbar.Button")
-								.OnClicked(this, &STutorialEditableText::HandleInsertBrowserLinkClicked)
+								.OnClicked(this, &STutorialEditableText::HandleInsertHyperLinkClicked)
 								[
 									SNew(STextBlock)
 									.TextStyle(FEditorStyle::Get(), "TutorialEditableText.Toolbar.Text")
-									.Text(LOCTEXT("HyperlinkInsertLabel", "Insert Hyperlink"))
+									.Text(this, &STutorialEditableText::GetHyperlinkButtonText)
 								]
 							]
 						]
@@ -487,181 +396,208 @@ void STutorialEditableText::HandleRichEditableTextCommitted(const FText& Text, E
 	OnTextCommitted.ExecuteIfBound(Text, Type);
 }
 
-void STutorialEditableText::HandleRichEditableTextCursorMoved(const FTextLocation& NewCursorPosition )
+static bool AreRunsTheSame(const TArray<TSharedRef<const IRun>>& Runs)
 {
-	// We can use GetRunUnderCursor to query the style of the text under the cursor
-	// We can then use this to update the toolbar
-	TSharedPtr<const IRun> Run = RichEditableTextBox->GetRunUnderCursor();
-	if(Run.IsValid() && Run->GetRunInfo().Name == TEXT("TextStyle"))
+	if(Runs.Num() == 0)
 	{
-		TextStyles.ExplodeRunInfo(Run->GetRunInfo(), ActiveFontFamily, FontSize, FontStyle, FontColor);
+		return false;
 	}
-}
-
-FText STutorialEditableText::GetActiveFontFamilyName() const
-{
-	return ActiveFontFamily->DisplayName;
-}
-
-void STutorialEditableText::OnActiveFontFamilyChanged(TSharedPtr<FTextStyles::FFontFamily> NewValue, ESelectInfo::Type)
-{
-	ActiveFontFamily = NewValue;
-	StyleSelectedText();
-}
-
-TSharedRef<SWidget> STutorialEditableText::GenerateFontFamilyComboEntry(TSharedPtr<FTextStyles::FFontFamily> SourceEntry)
-{
-	return SNew(STextBlock).Text(SourceEntry->DisplayName);
-}
 
-TOptional<uint8> STutorialEditableText::GetFontSize() const
-{
-	return FontSize;
-}
+	TSharedRef<const IRun> FirstRun = Runs[0];
+	for(const auto& Run : Runs)
+	{
+		if(Run != FirstRun)
+		{
+			if(Run->GetRunInfo().Name != FirstRun->GetRunInfo().Name)
+			{
+				return false;
+			}
 
-void STutorialEditableText::SetFontSize(uint8 NewValue, ETextCommit::Type)
-{		
-	FontSize = NewValue;
-	StyleSelectedText();
-}
+			for(const auto& MetaData : FirstRun->GetRunInfo().MetaData)
+			{
+				const FString* FoundMetaData = Run->GetRunInfo().MetaData.Find(MetaData.Key);
+				if(FoundMetaData == nullptr || *FoundMetaData != MetaData.Value)
+				{
+					return false;
+				}
+			}
+		}
+	}
 
-ESlateCheckBoxState::Type STutorialEditableText::IsFontStyleBold() const
-{
-	return (FontStyle & FTextStyles::EFontStyle::Bold) ? ESlateCheckBoxState::Checked : ESlateCheckBoxState::Unchecked;
+	return true;
 }
 
-void STutorialEditableText::OnFontStyleBoldChanged(ESlateCheckBoxState::Type InState)
+TSharedPtr<const IRun> STutorialEditableText::GetCurrentRun() const
 {
-	if(InState == ESlateCheckBoxState::Checked)
+	if(!RichEditableTextBox->GetSelectedText().IsEmpty())
 	{
-		FontStyle |= FTextStyles::EFontStyle::Bold;
+		const TArray<TSharedRef<const IRun>> Runs = RichEditableTextBox->GetSelectedRuns();
+		if(Runs.Num() == 1 || AreRunsTheSame(Runs))
+		{
+			return Runs[0];
+		}
 	}
 	else
 	{
-		FontStyle &= ~FTextStyles::EFontStyle::Bold;
+		return RichEditableTextBox->GetRunUnderCursor();
 	}
-	StyleSelectedText();
-}
 
-ESlateCheckBoxState::Type STutorialEditableText::IsFontStyleItalic() const
-{
-	return (FontStyle & FTextStyles::EFontStyle::Italic) ? ESlateCheckBoxState::Checked : ESlateCheckBoxState::Unchecked;
+	return TSharedPtr<const IRun>();
 }
 
-void STutorialEditableText::OnFontStyleItalicChanged(ESlateCheckBoxState::Type InState)
+void STutorialEditableText::HandleRichEditableTextCursorMoved(const FTextLocation& NewCursorPosition )
 {
-	if(InState == ESlateCheckBoxState::Checked)
+	TSharedPtr<const IRun> Run = GetCurrentRun();
+	
+	if(Run.IsValid())
 	{
-		FontStyle |= FTextStyles::EFontStyle::Italic;
+		if(Run->GetRunInfo().Name == TEXT("TextStyle"))
+		{
+			ActiveStyle = StylesAndNames[0];
+
+			FName StyleName = FTextStyleAndName::GetStyleFromRunInfo(Run->GetRunInfo());
+			for(const auto& StyleAndName : StylesAndNames)
+			{
+				if(StyleAndName->Style == StyleName)
+				{
+					ActiveStyle = StyleAndName;
+					break;
+				}
+			}
+		}
+		else if(Run->GetRunInfo().Name == TEXT("a"))
+		{
+			ActiveStyle = HyperlinkStyle;
+		}
+
+		FontComboBox->SetSelectedItem(ActiveStyle);
 	}
 	else
 	{
-		FontStyle &= ~FTextStyles::EFontStyle::Italic;
+		FontComboBox->SetSelectedItem(nullptr);
 	}
-	StyleSelectedText();
-}
-
-FLinearColor STutorialEditableText::GetFontColor() const
-{
-	return FontColor;
 }
 
-void STutorialEditableText::SetFontColor(FLinearColor NewValue)
+FText STutorialEditableText::GetActiveStyleName() const
 {
-	FontColor = NewValue;
-	StyleSelectedText();
+	return ActiveStyle.IsValid() ? ActiveStyle->DisplayName : FText();
 }
 
-FReply STutorialEditableText::OpenFontColorPicker()
+void STutorialEditableText::OnActiveStyleChanged(TSharedPtr<FTextStyleAndName> NewValue, ESelectInfo::Type SelectionType)
 {
-	FColorPickerArgs PickerArgs;
-	PickerArgs.bOnlyRefreshOnMouseUp = true;
-	PickerArgs.ParentWidget = AsShared();
-	PickerArgs.bUseAlpha = false;
-	PickerArgs.bOnlyRefreshOnOk = false;
-	PickerArgs.OnColorCommitted = FOnLinearColorValueChanged::CreateSP(this, &STutorialEditableText::SetFontColor);
-	PickerArgs.OnColorPickerCancelled = FOnColorPickerCancelled::CreateSP(this, &STutorialEditableText::SetFontColor);
-	PickerArgs.OnColorPickerWindowClosed = FOnWindowClosed::CreateSP(this, &STutorialEditableText::HandleColorPickerWindowClosed);
-	PickerArgs.InitialColorOverride = FontColor;
-
-	OpenColorPicker(PickerArgs);
-
-	bPickingColor = true;
-
-	return FReply::Handled();
+	ActiveStyle = NewValue;
+	if(SelectionType != ESelectInfo::Direct)
+	{
+		// style text if it was the user that made this selection
+		if(NewValue == HyperlinkStyle)
+		{
+			HandleHyperlinkComboOpened();
+			HyperlinkComboButton->SetIsOpen(true);
+		}
+		else
+		{
+			StyleSelectedText();
+		}
+	}
 }
 
-void STutorialEditableText::HandleColorPickerWindowClosed(const TSharedRef<SWindow>& /*InWindow*/)
+TSharedRef<SWidget> STutorialEditableText::GenerateStyleComboEntry(TSharedPtr<FTextStyleAndName> SourceEntry)
 {
-	bPickingColor = false;
+	return SNew(SBorder)
+		.BorderImage( FCoreStyle::Get().GetBrush( "NoBorder" ) )
+		.ForegroundColor(FCoreStyle::Get().GetSlateColor("InvertedForeground"))
+		[
+			SNew(STextBlock)
+			.Text(SourceEntry->DisplayName)
+			.TextStyle(&FEditorStyle::Get().GetWidgetStyle<FTextBlockStyle>(SourceEntry->Style))
+		];
 }
 
 void STutorialEditableText::StyleSelectedText()
 {
 	// Apply the current style to the selected text
 	// If no text is selected, then a new (empty) run will be inserted with the appropriate style
-	const FRunInfo RunInfo = TextStyles.CreateRunInfo(ActiveFontFamily, FontSize, FontStyle, FontColor);
-	const FTextBlockStyle TextBlockStyle = TextStyles.CreateTextBlockStyle(ActiveFontFamily, FontSize, FontStyle, FontColor);
-	RichEditableTextBox->ApplyToSelection(RunInfo, TextBlockStyle);
-	FSlateApplication::Get().SetKeyboardFocus(RichEditableTextBox, EKeyboardFocusCause::SetDirectly);
+	if(ActiveStyle.IsValid())
+	{
+		const FRunInfo RunInfo = ActiveStyle->CreateRunInfo();
+		const FTextBlockStyle TextBlockStyle = ActiveStyle->CreateTextBlockStyle();
+		RichEditableTextBox->ApplyToSelection(RunInfo, TextBlockStyle);
+		FSlateApplication::Get().SetKeyboardFocus(RichEditableTextBox, EKeyboardFocusCause::SetDirectly);
+	}
 }
 
+
 void STutorialEditableText::HandleHyperlinkComboOpened()
 {
+	HyperlinkURLTextBox->SetText(FText());
+	HyperlinkNameTextBlock->SetText(FText());
+
 	// Read any currently selected text, and use this as the default name of the hyperlink
 	FString SelectedText = RichEditableTextBox->GetSelectedText().ToString();
-	for(int32 SelectedTextIndex = 0; SelectedTextIndex < SelectedText.Len(); ++SelectedTextIndex)
+	if(SelectedText.Len() > 0)
 	{
-		if(FChar::IsLinebreak(SelectedText[SelectedTextIndex]))
+		for(int32 SelectedTextIndex = 0; SelectedTextIndex < SelectedText.Len(); ++SelectedTextIndex)
 		{
-			SelectedText = SelectedText.Left(SelectedTextIndex);
-			break;
+			if(FChar::IsLinebreak(SelectedText[SelectedTextIndex]))
+			{
+				SelectedText = SelectedText.Left(SelectedTextIndex);
+				break;
+			}
 		}
+		HyperlinkNameTextBlock->SetText(FText::FromString(SelectedText));
 	}
-	HyperlinkNameTextBox->SetText(FText::FromString(SelectedText));
 
-	// We can use GetRunUnderCursor to query whether the cursor is currently over a hyperlink
-	// If it is, we can use that as the default URL for the hyperlink
-	TSharedPtr<const IRun> Run = RichEditableTextBox->GetRunUnderCursor();
+	TSharedPtr<const IRun> Run = GetCurrentRun();
 	if(Run.IsValid() && Run->GetRunInfo().Name == TEXT("a"))
 	{
 		const FString* const URLUnderCursor = Run->GetRunInfo().MetaData.Find(TEXT("href"));
 		HyperlinkURLTextBox->SetText((URLUnderCursor) ? FText::FromString(*URLUnderCursor) : FText());
-	}
-	else
-	{
-		HyperlinkURLTextBox->SetText(FText());
-	}
+		FString RunText;
+		Run->AppendTextTo(RunText);
+		HyperlinkNameTextBlock->SetText(FText::FromString(RunText));
+	}	
 }
 
-FReply STutorialEditableText::HandleInsertBrowserLinkClicked()
+bool STutorialEditableText::IsHyperlinkComboEnabled() const
+{
+	return ActiveStyle == HyperlinkStyle;
+}
+
+FReply STutorialEditableText::HandleInsertHyperLinkClicked()
 {
 	HyperlinkComboButton->SetIsOpen(false);
 
-	const FText& Name = HyperlinkNameTextBox->GetText();
-	const FText& URL = HyperlinkURLTextBox->GetText();
-
-	// Create the correct meta-information for this run, so that valid source rich-text formatting can be generated for it
-	FRunInfo RunInfo(TEXT("a"));
-	RunInfo.MetaData.Add(TEXT("id"), TEXT("browser"));
-	RunInfo.MetaData.Add(TEXT("href"), URL.ToString());
-	RunInfo.MetaData.Add(TEXT("style"), TEXT("TutorialEditableText.Editor.Hyperlink"));
-
-	// Create the new run, and then insert it at the cursor position
-	TSharedRef<FSlateHyperlinkRun> HyperlinkRun = FSlateHyperlinkRun::Create(
-		RunInfo, 
-		MakeShareable(new FString(Name.ToString())), 
-		FEditorStyle::Get().GetWidgetStyle<FHyperlinkStyle>(FName(TEXT("TutorialEditableText.Editor.Hyperlink"))), 
-		OnBrowserLinkClicked
-		);
-	RichEditableTextBox->InsertRunAtCursor(HyperlinkRun);
+	const FText& Name = HyperlinkNameTextBlock->GetText();
+	if(!Name.IsEmpty())
+	{
+		const FText& URL = HyperlinkURLTextBox->GetText();
+
+		// Create the correct meta-information for this run, so that valid source rich-text formatting can be generated for it
+		FRunInfo RunInfo(TEXT("a"));
+		RunInfo.MetaData.Add(TEXT("id"), TEXT("browser"));
+		RunInfo.MetaData.Add(TEXT("href"), URL.ToString());
+		RunInfo.MetaData.Add(TEXT("style"), TEXT("Tutorials.Content.Hyperlink"));
+
+		// Create the new run, and then insert it at the cursor position
+		TSharedRef<FSlateHyperlinkRun> HyperlinkRun = FSlateHyperlinkRun::Create(
+			RunInfo, 
+			MakeShareable(new FString(Name.ToString())), 
+			FEditorStyle::Get().GetWidgetStyle<FHyperlinkStyle>(FName(TEXT("Tutorials.Content.Hyperlink"))), 
+			OnBrowserLinkClicked
+			);
+
+		// @todo: if the rich editable text box allowed us to replace a run that we found under the cursor (or returned a non-const
+		// instance of it) then we could edit the hyperlink here. This would mean the user does not need to select the whole hyperlink 
+		// to edit its URL.
+		RichEditableTextBox->InsertRunAtCursor(HyperlinkRun);
+	}
 
 	return FReply::Handled();
 }
 
 EVisibility STutorialEditableText::GetToolbarVisibility() const
 {
-	return FontComboBox->IsOpen() || HyperlinkComboButton->IsOpen() || bPickingColor || HasKeyboardFocus() || HasFocusedDescendants() ? EVisibility::Visible : EVisibility::Collapsed;
+	return FontComboBox->IsOpen() || HyperlinkComboButton->IsOpen() || HasKeyboardFocus() || HasFocusedDescendants() ? EVisibility::Visible : EVisibility::Collapsed;
 }
 
 ESlateCheckBoxState::Type STutorialEditableText::IsCreatingBrowserLink() const
@@ -729,4 +665,9 @@ void STutorialEditableText::OnCheckAssetLink(ESlateCheckBoxState::Type State)
 	}
 }
 
+FText STutorialEditableText::GetHyperlinkButtonText() const
+{
+	return bNewHyperlink ? LOCTEXT("HyperlinkInsertLabel", "Insert Hyperlink") : LOCTEXT("HyperlinkSetLabel", "Set Hyperlink");
+}
+
 #undef LOCTEXT_NAMESPACE
\ No newline at end of file
diff --git a/Engine/Source/Editor/IntroTutorials/Private/STutorialEditableText.h b/Engine/Source/Editor/IntroTutorials/Private/STutorialEditableText.h
--- a/Engine/Source/Editor/IntroTutorials/Private/STutorialEditableText.h
+++ b/Engine/Source/Editor/IntroTutorials/Private/STutorialEditableText.h
@@ -27,13 +27,10 @@ namespace EHyperlinkType
 	};
 }
 
- /**
-  * This is used in conjunction with the TextStyle decorator to allow arbitrary styling of text within a rich-text editor
-  * This struct defines a set of known font families, as well as providing some utility functions for converting the text style to and from a text layout run
-  */
-struct FTextStyles
+/** Text style and name to display in the UI */
+struct FTextStyleAndName
 {
-	/** Flags controlling which TTF or OTF font should be picked from the given font family */
+	/** Legacy: Flags controlling which TTF or OTF font should be picked from the given font family */
 	struct EFontStyle
 	{
 		typedef uint8 Flags;
@@ -45,154 +42,100 @@ struct FTextStyles
 		};
 	};
 
-	/** This struct defines a font family, which combines multiple TTF or OTF fonts into a single group, allowing text to be styled as bold or italic */
-	struct FFontFamily
-	{
-		FFontFamily(FText InDisplayName, FName InFamilyName, FName InRegularFont, FName InBoldFont, FName InItalicFont, FName InBoldItalicFont)
-			: DisplayName(InDisplayName)
-			, FamilyName(InFamilyName)
-			, RegularFont(InRegularFont)
-			, BoldFont(InBoldFont)
-			, ItalicFont(InItalicFont)
-			, BoldItalicFont(InBoldItalicFont)
-		{
-		}
-
-		/** Name used for this font family in the UI */
-		FText DisplayName;
-
-		/** Named used to identify this family from the TextStyle decorator */
-		FName FamilyName;
-
-		/** File paths to the fonts on disk */
-		FName RegularFont;
-		FName BoldFont;
-		FName ItalicFont;
-		FName BoldItalicFont;
-	};
+	FTextStyleAndName(FName InStyle, FText InDisplayName)
+		: Style(InStyle)
+		, DisplayName(InDisplayName)
+	{}
 
-	/** Convert the given text style into run meta-information, so that valid source rich-text formatting can be generated for it */
-	static FRunInfo CreateRunInfo(const TSharedPtr<FFontFamily>& InFontFamily, const uint8 InFontSize, const EFontStyle::Flags InFontStyle, const FLinearColor& InFontColor)
+	FRunInfo CreateRunInfo() const
 	{
-		FString FontStyleString;
-		if(InFontStyle == EFontStyle::Regular)
-		{
-			FontStyleString = TEXT("Regular");
-		}
-		else
-		{
-			if(InFontStyle & EFontStyle::Bold)
-			{
-				FontStyleString += "Bold";
-			}
-			if(InFontStyle & EFontStyle::Italic)
-			{
-				FontStyleString += "Italic";
-			}
-		}
-
 		FRunInfo RunInfo(TEXT("TextStyle"));
-		RunInfo.MetaData.Add(TEXT("FontFamily"), InFontFamily->FamilyName.ToString());
-		RunInfo.MetaData.Add(TEXT("FontSize"), FString::FromInt(InFontSize));
-		RunInfo.MetaData.Add(TEXT("FontStyle"), FontStyleString);
-		RunInfo.MetaData.Add(TEXT("FontColor"), InFontColor.ToString());
+		RunInfo.MetaData.Add(TEXT("Style"), Style.ToString());
 		return RunInfo;
 	}
 
-	/** Explode some run meta-information back out into its component text style parts */
-	void ExplodeRunInfo(const FRunInfo& InRunInfo, TSharedPtr<FFontFamily>& OutFontFamily, uint8& OutFontSize, EFontStyle::Flags& OutFontStyle, FLinearColor& OutFontColor) const
+	static FName GetStyleFromRunInfo(const FRunInfo& InRunInfo)
 	{
-		check(AvailableFontFamilies.Num());
-
-		const FString* const FontFamilyString = InRunInfo.MetaData.Find(TEXT("FontFamily"));
-		if(FontFamilyString)
+		const FString* const StyleString = InRunInfo.MetaData.Find(TEXT("Style"));
+		if(StyleString)
 		{
-			OutFontFamily = FindFontFamily(FName(**FontFamilyString));
+			return **StyleString;
 		}
-		if(!OutFontFamily.IsValid())
+		else
 		{
-			OutFontFamily = AvailableFontFamilies[0];
-		}
+			// legacy styling support - try to map older flexible styles to new fixed styles
 
-		OutFontSize = 11;
-		const FString* const FontSizeString = InRunInfo.MetaData.Find(TEXT("FontSize"));
-		if(FontSizeString)
-		{
-			OutFontSize = static_cast<uint8>(FPlatformString::Atoi(**FontSizeString));
-		}
+			int32 FontSize = 11;
+			const FString* const FontSizeString = InRunInfo.MetaData.Find(TEXT("FontSize"));
+			if(FontSizeString)
+			{
+				FontSize = static_cast<uint8>(FPlatformString::Atoi(**FontSizeString));
+			}
 
-		OutFontStyle = EFontStyle::Regular;
-		const FString* const FontStyleString = InRunInfo.MetaData.Find(TEXT("FontStyle"));
-		if(FontStyleString)
-		{
-			if(*FontStyleString == TEXT("Bold"))
+			if(FontSize > 24)
 			{
-				 OutFontStyle = EFontStyle::Bold;
+				return TEXT("Tutorials.Content.HeaderText2");
 			}
-			else if(*FontStyleString == TEXT("Italic"))
+			else if(FontSize > 11 && FontSize <= 24)
 			{
-				 OutFontStyle = EFontStyle::Italic;
+				return TEXT("Tutorials.Content.HeaderText1");
 			}
-			else if(*FontStyleString == TEXT("BoldItalic"))
+			else
 			{
-				 OutFontStyle = EFontStyle::Bold | EFontStyle::Italic;
+				EFontStyle::Flags FontStyle = EFontStyle::Regular;
+				const FString* const FontStyleString = InRunInfo.MetaData.Find(TEXT("FontStyle"));
+				if(FontStyleString)
+				{
+					if(*FontStyleString == TEXT("Bold"))
+					{
+						FontStyle = EFontStyle::Bold;
+					}
+					else if(*FontStyleString == TEXT("Italic"))
+					{
+						FontStyle = EFontStyle::Italic;
+					}
+					else if(*FontStyleString == TEXT("BoldItalic"))
+					{
+						FontStyle = EFontStyle::Bold | EFontStyle::Italic;
+					}
+				}
+
+				FLinearColor FontColor = FLinearColor::Black;
+				const FString* const FontColorString = InRunInfo.MetaData.Find(TEXT("FontColor"));
+				if(FontColorString && !FontColor.InitFromString(*FontColorString))
+				{
+					FontColor = FLinearColor::Black;
+				}
+
+				if(FontStyle != EFontStyle::Regular || FontColor != FLinearColor::Black)
+				{
+					return TEXT("Tutorials.Content.TextBold");
+				}
+				else
+				{
+					return TEXT("Tutorials.Content.Text");
+				}
 			}
 		}
 
-		OutFontColor = FLinearColor::Black;
-		const FString* const FontColorString = InRunInfo.MetaData.Find(TEXT("FontColor"));
-		if(FontColorString && !OutFontColor.InitFromString(*FontColorString))
-		{
-			OutFontColor = FLinearColor::Black;
-		}
+		return NAME_None;
 	}
 
-	/** Convert the given text style into a text block style for use by Slate */
-	static FTextBlockStyle CreateTextBlockStyle(const TSharedPtr<FFontFamily>& InFontFamily, const uint8 InFontSize, const EFontStyle::Flags InFontStyle, const FLinearColor& InFontColor)
+	FTextBlockStyle CreateTextBlockStyle() const
 	{
-		FName FontName = InFontFamily->RegularFont;
-		if((InFontStyle & EFontStyle::Bold) && (InFontStyle & EFontStyle::Italic))
-		{
-			FontName = InFontFamily->BoldItalicFont;
-		}
-		else if(InFontStyle & EFontStyle::Bold)
-		{
-			FontName = InFontFamily->BoldFont;
-		}
-		else if(InFontStyle & EFontStyle::Italic)
-		{
-			FontName = InFontFamily->ItalicFont;
-		}
-
-		FTextBlockStyle TextBlockStyle;
-		TextBlockStyle.SetFontName(FontName);
-		TextBlockStyle.SetFontSize(InFontSize);
-		TextBlockStyle.SetColorAndOpacity(InFontColor);
-		return TextBlockStyle;
+		return FEditorStyle::Get().GetWidgetStyle<FTextBlockStyle>(Style);
 	}
 
-	/** Convert the given run meta-information into a text block style for use by Slate */
-	FTextBlockStyle CreateTextBlockStyle(const FRunInfo& InRunInfo) const
+	static FTextBlockStyle CreateTextBlockStyle(const FRunInfo& InRunInfo)
 	{
-		TSharedPtr<FFontFamily> FontFamily;
-		uint8 FontSize;
-		EFontStyle::Flags FontStyle;
-		FLinearColor FontColor;
-		ExplodeRunInfo(InRunInfo, FontFamily, FontSize, FontStyle, FontColor);
-		return CreateTextBlockStyle(FontFamily, FontSize, FontStyle, FontColor);
+		return FEditorStyle::Get().GetWidgetStyle<FTextBlockStyle>(GetStyleFromRunInfo(InRunInfo));
 	}
 
-	/** Try and find a font family with the given name */
-	TSharedPtr<FFontFamily> FindFontFamily(const FName InFamilyName) const
-	{
-		const TSharedPtr<FFontFamily>* const FoundFontFamily = AvailableFontFamilies.FindByPredicate([InFamilyName](TSharedPtr<FFontFamily>& Entry) -> bool
-		{
-			return Entry->FamilyName == InFamilyName;
-		});
-		return (FoundFontFamily) ? *FoundFontFamily : nullptr;
-	}
+	/** The style identifier */
+	FName Style;
 
-	TArray<TSharedPtr<FFontFamily>> AvailableFontFamilies;
+	/** The text to display for this style */
+	FText DisplayName;
 };
 
 /**
@@ -203,9 +146,9 @@ class FTextStyleDecorator : public ITextDecorator
 {
 public:
 
-	static TSharedRef<FTextStyleDecorator> Create(const FTextStyles& InTextStyles)
+	static TSharedRef<FTextStyleDecorator> Create()
 	{
-		return MakeShareable(new FTextStyleDecorator(InTextStyles));
+		return MakeShareable(new FTextStyleDecorator());
 	}
 
 	virtual ~FTextStyleDecorator()
@@ -230,17 +173,8 @@ public:
 		*InOutModelText += OriginalText.Mid(RunParseResult.ContentRange.BeginIndex, RunParseResult.ContentRange.EndIndex - RunParseResult.ContentRange.BeginIndex);
 		ModelRange.EndIndex = InOutModelText->Len();
 
-		return FSlateTextRun::Create(RunInfo, InOutModelText, TextStyles.CreateTextBlockStyle(RunInfo), ModelRange);
+		return FSlateTextRun::Create(RunInfo, InOutModelText, FTextStyleAndName::CreateTextBlockStyle(RunInfo), ModelRange);
 	}
-
-private:
-
-	FTextStyleDecorator(const FTextStyles& InTextStyles)
-		: TextStyles(InTextStyles)
-	{
-	}
-
-	FTextStyles TextStyles;
 };
 
 
@@ -268,43 +202,27 @@ public:
 
 protected:
 	
+	TSharedPtr<const IRun> GetCurrentRun() const;
+
 	void HandleRichEditableTextChanged(const FText& Text);
 
 	void HandleRichEditableTextCommitted(const FText& Text, ETextCommit::Type Type);
 
 	void HandleRichEditableTextCursorMoved(const FTextLocation& NewCursorPosition );
 
-	FText GetActiveFontFamilyName() const;
-
-	void OnActiveFontFamilyChanged(TSharedPtr<FTextStyles::FFontFamily> NewValue, ESelectInfo::Type);
-
-	TSharedRef<SWidget> GenerateFontFamilyComboEntry(TSharedPtr<FTextStyles::FFontFamily> SourceEntry);
-
-	TOptional<uint8> GetFontSize() const;
-
-	void SetFontSize(uint8 NewValue, ETextCommit::Type);
-
-	ESlateCheckBoxState::Type IsFontStyleBold() const;
-
-	void OnFontStyleBoldChanged(ESlateCheckBoxState::Type InState);
+	FText GetActiveStyleName() const;
 
-	ESlateCheckBoxState::Type IsFontStyleItalic() const;
+	void OnActiveStyleChanged(TSharedPtr<FTextStyleAndName> NewValue, ESelectInfo::Type);
 
-	void OnFontStyleItalicChanged(ESlateCheckBoxState::Type InState);
-
-	FLinearColor GetFontColor() const;
-
-	void SetFontColor(FLinearColor NewValue);
-
-	FReply OpenFontColorPicker();
-
-	void HandleColorPickerWindowClosed(const TSharedRef<SWindow>& InWindow);
+	TSharedRef<SWidget> GenerateStyleComboEntry(TSharedPtr<FTextStyleAndName> SourceEntry);
 
 	void StyleSelectedText();
 
 	void HandleHyperlinkComboOpened();
 
-	FReply HandleInsertBrowserLinkClicked();
+	bool IsHyperlinkComboEnabled() const;
+
+	FReply HandleInsertHyperLinkClicked();
 
 	EVisibility GetToolbarVisibility() const;
 
@@ -328,6 +246,8 @@ protected:
 
 	void OnCheckAssetLink(ESlateCheckBoxState::Type State);
 
+	FText GetHyperlinkButtonText() const;
+
 protected:
 	TSharedPtr<SMultiLineEditableTextBox> RichEditableTextBox;
 
@@ -338,20 +258,19 @@ protected:
 	FSlateHyperlinkRun::FOnClick OnAssetLinkClicked;
 
 	TSharedPtr<SComboButton> HyperlinkComboButton;
-	TSharedPtr<SComboBox<TSharedPtr<FTextStyles::FFontFamily>>> FontComboBox;
-	TSharedPtr<SEditableTextBox> HyperlinkNameTextBox;
+	TSharedPtr<SComboBox<TSharedPtr<FTextStyleAndName>>> FontComboBox;
+	TSharedPtr<STextBlock> HyperlinkNameTextBlock;
 	TSharedPtr<SEditableTextBox> HyperlinkURLTextBox;
 
-	FTextStyles TextStyles;
+	TSharedPtr<FTextStyleAndName> ActiveStyle;
+	TSharedPtr<FTextStyleAndName> HyperlinkStyle;
 
-	TSharedPtr<FTextStyles::FFontFamily> ActiveFontFamily;
-	uint8 FontSize;
-	FTextStyles::EFontStyle::Flags FontStyle;
-	FLinearColor FontColor;
-	bool bPickingColor;
+	TArray<TSharedPtr<FTextStyleAndName>> StylesAndNames;
 
 	FOnTextCommitted OnTextCommitted;
 	FOnTextChanged OnTextChanged;
 
 	EHyperlinkType::Type CurrentHyperlinkType;
+
+	bool bNewHyperlink;
 };
