Created a metadata type that details blueprint information for the tutorial widgets.
Revised the variable node and base node to populate and insert such a tag.

diff --git a/Engine/Source/Editor/GraphEditor/GraphEditor.Build.cs b/Engine/Source/Editor/GraphEditor/GraphEditor.Build.cs
--- a/Engine/Source/Editor/GraphEditor/GraphEditor.Build.cs
+++ b/Engine/Source/Editor/GraphEditor/GraphEditor.Build.cs
@@ -19,6 +19,12 @@ public class GraphEditor : ModuleRules
 			}
 		);
 
+        PublicIncludePathModuleNames.AddRange(
+            new string[] {                
+                "IntroTutorials"
+            }
+        );
+
 		PrivateIncludePathModuleNames.AddRange(
 			new string[] { 
 				"EditorWidgets"
@@ -42,7 +48,7 @@ public class GraphEditor : ModuleRules
 				"AnimGraph",
 				"Documentation",
 				"RenderCore",
-				"RHI"
+				"RHI",
 			}
 		);
 
diff --git a/Engine/Source/Editor/GraphEditor/Private/KismetNodes/SGraphNodeK2Base.cpp b/Engine/Source/Editor/GraphEditor/Private/KismetNodes/SGraphNodeK2Base.cpp
--- a/Engine/Source/Editor/GraphEditor/Private/KismetNodes/SGraphNodeK2Base.cpp
+++ b/Engine/Source/Editor/GraphEditor/Private/KismetNodes/SGraphNodeK2Base.cpp
@@ -8,6 +8,7 @@
 #include "Editor/UnrealEd/Public/Kismet2/BlueprintEditorUtils.h"
 #include "KismetNodeInfoContext.h"
 #include "IDocumentation.h"
+#include "TutorialMetaData.h"
 
 #define LOCTEXT_NAMESPACE "SGraphNodeK2Base"
 
@@ -45,19 +46,9 @@ void SGraphNodeK2Base::UpdateCompactNode()
 	}
 
 	// Setup a tag for this node
-	FString TagName;
-	if (GraphNode != nullptr)
-	{
-		// We want the name of the blueprint as our name - we can find the node from the GUID
-		UObject* Package = GraphNode->GetOutermost();
-		UObject* LastOuter = GraphNode->GetOuter();
-		while (LastOuter->GetOuter() != Package)
-		{
-			LastOuter = LastOuter->GetOuter();
-		}
-		TagName = FString::Printf(TEXT("GraphNode,%s,%s"), *LastOuter->GetFullName(), *GraphNode->NodeGuid.ToString());
-	}
-	
+	FBlueprintGraphNodeMetaData TagMeta(TEXT("Graphnode"), TEXT("None"), FGuid());
+	PopulateMetaTag(&TagMeta);
+
 	//
 	//             ______________________
 	//            | (>) L |      | R (>) |
@@ -84,7 +75,7 @@ void SGraphNodeK2Base::UpdateCompactNode()
 			// NODE CONTENT AREA
 			SNew(SOverlay)
 			.ToolTip( NodeToolTip.ToSharedRef() )
-			.Tag(*TagName)
+			.AddMetaData<FBlueprintGraphNodeMetaData>(TagMeta)
 			+SOverlay::Slot()
 			[
 				SNew(SImage)
@@ -525,4 +516,47 @@ void SGraphNodeK2Base::PerformSecondPassLayout(const TMap< UObject*, TSharedRef<
 	PositionThisNodeBetweenOtherNodes(NodeToWidgetLookup, PrevNodes, NextNodes, Height);
 }
 
+void SGraphNodeK2Base::PopulateMetaTag(FBlueprintGraphNodeMetaData* TagMeta) const
+{
+	if (GraphNode != nullptr)
+	{
+		// We want the name of the blueprint as our name - we can find the node from the GUID
+		UObject* Package = GraphNode->GetOutermost();
+		UObject* LastOuter = GraphNode->GetOuter();
+		while (LastOuter->GetOuter() != Package)
+		{
+			LastOuter = LastOuter->GetOuter();
+		}
+		TagMeta->Tag = FName(*FString::Printf(TEXT("GraphNode_%s_%s"), *LastOuter->GetFullName(), *GraphNode->NodeGuid.ToString()));
+		TagMeta->OuterName = LastOuter->GetFullName();
+		TagMeta->GUID = GraphNode->NodeGuid;
+		FString TitleText = GetNodeCompactTitle().ToString();
+		if (GraphNode->Pins.Num() > 0)
+		{
+			const UEdGraphSchema_K2* K2Schema = GetDefault<UEdGraphSchema_K2>();
+			for (int iPin = 0; iPin < GraphNode->Pins.Num(); iPin++)
+			{
+				if (GraphNode->Pins[iPin]->PinType.PinCategory != K2Schema->PC_Exec)
+				{
+					// Include the title if there is one
+					if (TitleText.IsEmpty())
+					{
+						TagMeta->FriendlyName = FString::Printf(TEXT("%s in %s"), *GraphNode->Pins[iPin]->PinName, *TagMeta->OuterName);
+					}
+					else
+					{
+						TagMeta->FriendlyName = FString::Printf(TEXT("%s %s in %s"), *TitleText, *GraphNode->Pins[iPin]->PinName, *TagMeta->OuterName);
+					}
+
+					break;
+				}
+			}
+		}
+		else
+		{
+			TagMeta->FriendlyName = TitleText;
+		}
+	}
+}
+
 #undef LOCTEXT_NAMESPACE
diff --git a/Engine/Source/Editor/GraphEditor/Private/KismetNodes/SGraphNodeK2Base.h b/Engine/Source/Editor/GraphEditor/Private/KismetNodes/SGraphNodeK2Base.h
--- a/Engine/Source/Editor/GraphEditor/Private/KismetNodes/SGraphNodeK2Base.h
+++ b/Engine/Source/Editor/GraphEditor/Private/KismetNodes/SGraphNodeK2Base.h
@@ -32,6 +32,9 @@ protected :
 	/** Retrieves text to tack on to the top of the tooltip (above the standard text) */
 	FText GetToolTipHeading() const;
 
+	/* Populate a meta data tag with information about this graph node */
+	void PopulateMetaTag(class FBlueprintGraphNodeMetaData* TagMeta) const;
+
 protected:
 	static const FLinearColor BreakpointHitColor;
 	static const FLinearColor LatentBubbleColor;
diff --git a/Engine/Source/Editor/GraphEditor/Private/KismetNodes/SGraphNodeK2Var.cpp b/Engine/Source/Editor/GraphEditor/Private/KismetNodes/SGraphNodeK2Var.cpp
--- a/Engine/Source/Editor/GraphEditor/Private/KismetNodes/SGraphNodeK2Var.cpp
+++ b/Engine/Source/Editor/GraphEditor/Private/KismetNodes/SGraphNodeK2Var.cpp
@@ -8,6 +8,7 @@
 #include "ClassIconFinder.h"
 #include "IDocumentation.h"
 #include "BlueprintEditorUtils.h"
+#include "TutorialMetaData.h"
 
 void SGraphNodeK2Var::Construct( const FArguments& InArgs, UK2Node* InNode )
 {
@@ -160,18 +161,8 @@ void SGraphNodeK2Var::UpdateGraphNode()
 	TSharedPtr<SWidget> ErrorText = SetupErrorReporting();
 
 	// Setup a tag for this node
-	FString TagName;
-	if (GraphNode != nullptr)
-	{
-		// We want the name of the blueprint as our name - we can find the node from the GUID
-		UObject* Package = GraphNode->GetOutermost();
-		UObject* LastOuter = GraphNode->GetOuter();
-		while (LastOuter->GetOuter() != Package)
-		{
-			LastOuter = LastOuter->GetOuter();
-		}
-		TagName = FString::Printf(TEXT("GraphNode,%s,%s"), *LastOuter->GetFullName(), *GraphNode->NodeGuid.ToString());
-	}
+	FBlueprintGraphNodeMetaData TagMeta(TEXT("Graphnode"), TEXT("None"), FGuid());	
+	PopulateMetaTag(&TagMeta, TitleText);
 
 	//             ________________
 	//            | (>) L |  R (>) |
@@ -197,7 +188,7 @@ void SGraphNodeK2Var::UpdateGraphNode()
 		+SVerticalBox::Slot()
 		[
 			SNew(SOverlay)
-			.Tag(*TagName)
+			.AddMetaData<FBlueprintGraphNodeMetaData>(TagMeta)
 			+ SOverlay::Slot()
 			[
 				SNew(SImage)
@@ -273,4 +264,46 @@ void SGraphNodeK2Var::UpdateGraphNode()
 const FSlateBrush* SGraphNodeK2Var::GetShadowBrush(bool bSelected) const
 {
 	return bSelected ? FEditorStyle::GetBrush(TEXT("Graph.VarNode.ShadowSelected")) : FEditorStyle::GetBrush(TEXT("Graph.VarNode.Shadow"));
-}
\ No newline at end of file
+}
+
+void SGraphNodeK2Var::PopulateMetaTag(FBlueprintGraphNodeMetaData* TagMeta, const FText &TitleText) const
+{
+	if (GraphNode != nullptr)
+	{
+		// We want the name of the blueprint as our name - we can find the node from the GUID
+		UObject* Package = GraphNode->GetOutermost();
+		UObject* LastOuter = GraphNode->GetOuter();
+		while (LastOuter->GetOuter() != Package)
+		{
+			LastOuter = LastOuter->GetOuter();
+		}
+		TagMeta->Tag = FName(*FString::Printf(TEXT("GraphNode_%s_%s"), *LastOuter->GetFullName(), *GraphNode->NodeGuid.ToString()));
+		TagMeta->OuterName = LastOuter->GetFullName();
+		TagMeta->GUID = GraphNode->NodeGuid;
+		if (GraphNode->Pins.Num() > 0)
+		{
+			const UEdGraphSchema_K2* K2Schema = GetDefault<UEdGraphSchema_K2>();
+			for (int iPin = 0; iPin < GraphNode->Pins.Num(); iPin++)
+			{
+				if (GraphNode->Pins[iPin]->PinType.PinCategory != K2Schema->PC_Exec)
+				{
+					// Include the title if there is one
+					if (TitleText.IsEmpty())
+					{
+						TagMeta->FriendlyName = FString::Printf(TEXT("%s in %s"), *GraphNode->Pins[iPin]->PinName, *TagMeta->OuterName);
+					}
+					else
+					{
+						TagMeta->FriendlyName = FString::Printf(TEXT("%s %s in %s"), *TitleText.ToString(), *GraphNode->Pins[iPin]->PinName, *TagMeta->OuterName);
+					}
+
+					break;
+				}
+			}
+		}
+		else
+		{
+			TagMeta->FriendlyName = TitleText.ToString();
+		}
+	}
+}
diff --git a/Engine/Source/Editor/GraphEditor/Private/KismetNodes/SGraphNodeK2Var.h b/Engine/Source/Editor/GraphEditor/Private/KismetNodes/SGraphNodeK2Var.h
--- a/Engine/Source/Editor/GraphEditor/Private/KismetNodes/SGraphNodeK2Var.h
+++ b/Engine/Source/Editor/GraphEditor/Private/KismetNodes/SGraphNodeK2Var.h
@@ -18,6 +18,9 @@ public:
 	virtual void UpdateGraphNode() override;
 	// End of SGraphNode interface
 
+	/* Populates a meta tag with information about the graph node */
+	void PopulateMetaTag(class FBlueprintGraphNodeMetaData* TagMeta, const FText &TitleText) const;
+
 protected:
 	FSlateColor GetVariableColor() const;
 };
diff --git a/Engine/Source/Editor/IntroTutorials/Private/STutorialOverlay.cpp b/Engine/Source/Editor/IntroTutorials/Private/STutorialOverlay.cpp
--- a/Engine/Source/Editor/IntroTutorials/Private/STutorialOverlay.cpp
+++ b/Engine/Source/Editor/IntroTutorials/Private/STutorialOverlay.cpp
@@ -104,8 +104,7 @@ int32 STutorialOverlay::OnPaint( const FPaintArgs& Args, const FGeometry& Allott
 int32 STutorialOverlay::TraverseWidgets(TSharedRef<SWidget> InWidget, const FGeometry& InGeometry, const FSlateRect& MyClippingRect, FSlateWindowElementList& OutDrawElements, int32 LayerId) const
 {
 	TSharedPtr<FTagMetaData> MetaData = InWidget->GetMetaData<FTagMetaData>();
-
-	const FName Tag = InWidget->GetTag();
+	const FName Tag = (MetaData.IsValid() && MetaData->Tag.IsValid()) ? MetaData->Tag : InWidget->GetTag();	
 	if(Tag != NAME_None || MetaData.IsValid())
 	{
 		// we are a named widget - ask it to draw
@@ -180,31 +179,22 @@ void STutorialOverlay::FocusOnAnyBlueprintNodes(const FTutorialWidgetContent &Wi
 	{
 		return;
 	}
-	FString IdentString = WidgetContent.WidgetAnchor.WrapperIdentifier.ToString();
-	TArray<FString> Tokens;
-	IdentString.ParseIntoArray(&Tokens, TEXT(","), true);
-	if (Tokens.Num() > 1 )
+	FString Name = WidgetContent.WidgetAnchor.OuterName;
+	int32 NameIndex;
+	Name.FindLastChar(TEXT('.'), NameIndex);
+	FString BlueprintName = Name.RightChop(NameIndex + 1);
+	UBlueprint* Blueprint = FindObject<UBlueprint>(ANY_PACKAGE, *BlueprintName);
+	// If we find a blueprint
+	if (Blueprint != nullptr)
 	{
-		FString Name = Tokens[1];
-
-		// I'm sure this isnt how to do this. But probably this will all come from meta data anyway
-		int32 NameIndex;
-		Name.FindLastChar(TEXT('.'), NameIndex);
-		FString BlueprintName = Name.RightChop(NameIndex+1);		
-		UBlueprint* Blueprint = FindObject<UBlueprint>(ANY_PACKAGE, *BlueprintName);
-		// If we find a blueprint
-		if (Blueprint != nullptr)
+		// Try to grab guid
+		FGuid NodeGuid;
+		FGuid::Parse(WidgetContent.WidgetAnchor.GUIDString, NodeGuid);
+		UEdGraphNode* OutNode = NULL;
+		if (UEdGraphNode* GraphNode = FBlueprintEditorUtils::GetNodeByGUID(Blueprint, NodeGuid))
 		{
-			// Try to grab guid
-			FGuid NodeGuid;
-			FGuid::Parse(Tokens[Tokens.Num() - 1], NodeGuid);
-			UEdGraphNode* OutNode = NULL;
-			if (UEdGraphNode* GraphNode = FBlueprintEditorUtils::GetNodeByGUID(Blueprint, NodeGuid))
-			{
-				FKismetEditorUtilities::BringKismetToFocusAttentionOnObject(GraphNode, false);
-			}
+			FKismetEditorUtilities::BringKismetToFocusAttentionOnObject(GraphNode, false);
 		}
-
 	}
 }
 
diff --git a/Engine/Source/Editor/IntroTutorials/Private/TutorialStructCustomization.cpp b/Engine/Source/Editor/IntroTutorials/Private/TutorialStructCustomization.cpp
--- a/Engine/Source/Editor/IntroTutorials/Private/TutorialStructCustomization.cpp
+++ b/Engine/Source/Editor/IntroTutorials/Private/TutorialStructCustomization.cpp
@@ -4,6 +4,8 @@
 #include "TutorialStructCustomization.h"
 #include "EditorTutorial.h"
 #include "STutorialEditableText.h"
+#include "ISlateMetaData.h"
+#include "TutorialMetaData.h"
 
 #define LOCTEXT_NAMESPACE "TutorialStructCustomization"
 
@@ -201,6 +203,7 @@ private:
 	virtual void Tick( const FGeometry& AllottedGeometry, const double InCurrentTime, const float InDeltaTime ) override
 	{
 		PickedWidgetName = NAME_None;
+		PickedAllMetaData.Reset();
 
 		FWidgetPath Path = FSlateApplication::Get().LocateWindowUnderMouse(FSlateApplication::Get().GetCursorPos(), FSlateApplication::Get().GetInteractiveTopLevelWindows(), true);
 
@@ -218,6 +221,10 @@ private:
 				if(MetaData.IsValid())
 				{
 					PickedWidgetName = MetaData->Tag;
+					if (PickedWidgetName != NAME_None)
+					{
+						PickedAllMetaData = PathWidget->GetAllMetaData<FTagMetaData>();
+					}
 					break;
 				}
 			}		
@@ -247,6 +254,24 @@ private:
 			TSharedPtr<IPropertyHandle> TypeProperty = StructPropertyHandle->GetChildHandle(GET_MEMBER_NAME_CHECKED(FTutorialContentAnchor, Type));
 			TypeProperty->SetValue((uint8)ETutorialAnchorIdentifier::NamedWidget);
 			
+			for (const auto& MetaDataEntry : PickedAllMetaData)
+			{
+				if (MetaDataEntry->IsOfType<FBlueprintGraphNodeMetaData>())
+				{
+					TSharedRef<FBlueprintGraphNodeMetaData> GraphNodeMeta = StaticCastSharedRef<FBlueprintGraphNodeMetaData>(MetaDataEntry);
+					StructPropertyHandle->GetChildHandle(GET_MEMBER_NAME_CHECKED(FTutorialContentAnchor, TagIdent))->SetValue(GraphNodeMeta->Tag);
+					StructPropertyHandle->GetChildHandle(GET_MEMBER_NAME_CHECKED(FTutorialContentAnchor, GUIDString))->SetValue(GraphNodeMeta->GUID.ToString());
+					StructPropertyHandle->GetChildHandle(GET_MEMBER_NAME_CHECKED(FTutorialContentAnchor, OuterName))->SetValue(GraphNodeMeta->OuterName);
+					StructPropertyHandle->GetChildHandle(GET_MEMBER_NAME_CHECKED(FTutorialContentAnchor, FriendlyName))->SetValue(GraphNodeMeta->FriendlyName);
+				}
+				else
+				{
+					TSharedRef<FTagMetaData> GraphNodeMeta = StaticCastSharedRef<FTagMetaData>(MetaDataEntry);
+					TSharedPtr<IPropertyHandle> MetaDataPropertyIdent = StructPropertyHandle->GetChildHandle(GET_MEMBER_NAME_CHECKED(FTutorialContentAnchor, TagIdent));
+					MetaDataPropertyIdent->SetValue(GraphNodeMeta->Tag);
+				}
+			}
+			
 			PickedWidgetName = NAME_None;
 
 			if(ParentWindow.IsValid())
@@ -276,6 +301,9 @@ private:
 
 	/** The widget name we are picking */
 	FName PickedWidgetName;
+
+	/* The metadata for the widget we are picking */
+	TArray<TSharedRef<FTagMetaData>> PickedAllMetaData;
 };
 
 /** Widget used to launch a 'picking' session */
@@ -375,15 +403,12 @@ private:
 		{
 			return FText::FromName(*WidgetName);
 		}
-		TArray<FString> Tokens;
-		WidgetName.ParseIntoArray(&Tokens, TEXT(","),true);
-		if (Tokens.Num() == 3)
-		{
-			return FText::FromName(*Tokens[1]);
-		}
-		else if (Tokens.Num() == 3)
+
+		FString FriendlyName;
+		StructPropertyHandle->GetChildHandle(GET_MEMBER_NAME_CHECKED(FTutorialContentAnchor, FriendlyName))->GetValue(FriendlyName);
+		if (FriendlyName.IsEmpty()==false)
 		{
-			return FText::FromName(*Tokens[0]);
+			return FText::FromName(*FriendlyName);
 		}
 		return FText::FromName(*WidgetName);
 	}
diff --git a/Engine/Source/Editor/IntroTutorials/Public/EditorTutorial.h b/Engine/Source/Editor/IntroTutorials/Public/EditorTutorial.h
--- a/Engine/Source/Editor/IntroTutorials/Public/EditorTutorial.h
+++ b/Engine/Source/Editor/IntroTutorials/Public/EditorTutorial.h
@@ -118,6 +118,16 @@ struct INTROTUTORIALS_API FTutorialContentAnchor
 	/** Whether to draw an animated highlight around the widget */
 	UPROPERTY(EditAnywhere, Category="Anchor")
 	bool bDrawHighlight;
+
+	//FBlueprintGraphNodeMetaData Data
+	UPROPERTY(EditAnywhere, Category = "AnchorMeta|FBlueprintGraphNodeMetaData")
+	FString TagIdent;
+	UPROPERTY(EditAnywhere, Category = "AnchorMeta|FBlueprintGraphNodeMetaData")
+ 	FString GUIDString;
+	UPROPERTY(EditAnywhere, Category = "AnchorMeta|FBlueprintGraphNodeMetaData")
+	FString OuterName;
+	UPROPERTY(EditAnywhere, Category = "AnchorMeta|FBlueprintGraphNodeMetaData")
+	FString FriendlyName;
 };
 
 /** Content that is displayed relative to a widget */
diff --git a/Engine/Source/Editor/IntroTutorials/Public/TutorialMetaData.h b/Engine/Source/Editor/IntroTutorials/Public/TutorialMetaData.h
new file mode 100644
--- /dev/null
+++ b/Engine/Source/Editor/IntroTutorials/Public/TutorialMetaData.h
@@ -0,0 +1,28 @@
+// Copyright 1998-2014 Epic Games, Inc. All Rights Reserved.
+
+#pragma once
+
+#include "ISlateMetaData.h"
+
+class INTROTUTORIALS_API FBlueprintGraphNodeMetaData : public FTagMetaData
+{
+public:
+	SLATE_METADATA_TYPE(FBlueprintGraphNodeMetaData, FTagMetaData)
+
+		FBlueprintGraphNodeMetaData(FName InTag, FString InOuterName, FGuid InGUID)
+		: FTagMetaData(InTag)
+		, GUID(InGUID)
+		, OuterName(InOuterName)
+		{
+		}
+
+	
+	/** GUID for the node */
+	FGuid GUID;
+
+	/** Name of the outer (which should be the blueprint) */
+	FString OuterName;
+
+	/** User friendly display name */
+	FString FriendlyName;
+};
\ No newline at end of file
