Added more hyperlink types for tutorials

Added asset, code, tutorial and UDN hyperlinks. Tweaked UI to use combo box instead of the large amount of radio buttons it was using before.
Made sure importer & content use correct text styles.
Defaulted to rich text for content.

diff --git a/Engine/Source/Editor/Documentation/Public/IDocumentation.h b/Engine/Source/Editor/Documentation/Public/IDocumentation.h
--- a/Engine/Source/Editor/Documentation/Public/IDocumentation.h
+++ b/Engine/Source/Editor/Documentation/Public/IDocumentation.h
@@ -58,6 +58,13 @@ public:
 		return *this;
 	}
 
+	/** Set the hyperlink button style for this documentation */
+	FDocumentationStyle& HyperlinkStyle(const FName& InName) 
+	{
+		HyperlinkStyleName = InName;
+		return *this;
+	}
+
 	/** Set the hyperlink button style for this documentation */
 	FDocumentationStyle& HyperlinkButtonStyle(const FName& InName) 
 	{
diff --git a/Engine/Source/Editor/IntroTutorials/Private/SEditorTutorialImporter.cpp b/Engine/Source/Editor/IntroTutorials/Private/SEditorTutorialImporter.cpp
--- a/Engine/Source/Editor/IntroTutorials/Private/SEditorTutorialImporter.cpp
+++ b/Engine/Source/Editor/IntroTutorials/Private/SEditorTutorialImporter.cpp
@@ -93,7 +93,18 @@ void SEditorTutorialImporter::Import(UEditorTutorial* InTutorial)
 {
 	if(!UDNPath.IsEmpty())
 	{
-		TSharedRef<IDocumentationPage> Page = IDocumentation::Get()->GetPage(UDNPath.ToString(), TSharedPtr<FParserConfiguration>(), FDocumentationStyle());
+		FDocumentationStyle DocumentationStyle;
+		DocumentationStyle
+			.ContentStyle(TEXT("Tutorials.Content.Text"))
+			.BoldContentStyle(TEXT("Tutorials.Content.TextBold"))
+			.NumberedContentStyle(TEXT("Tutorials.Content.Text"))
+			.Header1Style(TEXT("Tutorials.Content.HeaderText1"))
+			.Header2Style(TEXT("Tutorials.Content.HeaderText2"))
+			.HyperlinkStyle(TEXT("Tutorials.Content.Hyperlink"))
+			.HyperlinkTextStyle(TEXT("Tutorials.Content.HyperlinkText"))
+			.SeparatorStyle(TEXT("Tutorials.Separator"));
+
+		TSharedRef<IDocumentationPage> Page = IDocumentation::Get()->GetPage(UDNPath.ToString(), TSharedPtr<FParserConfiguration>(), DocumentationStyle);
 		InTutorial->Title = Page->GetTitle();
 
 		TArray<FExcerpt> Excerpts;
@@ -112,7 +123,7 @@ void SEditorTutorialImporter::Import(UEditorTutorial* InTutorial)
 				FString* TitleStringPtr = Excerpt.Variables.Find(TEXT("StageTitle"));
 				if(TitleStringPtr != nullptr)
 				{
-					RichText += FString::Printf(TEXT("<TextStyle FontFamily=\"Roboto\" FontSize=\"24\" FontStyle=\"Bold\" FontColor=\"(R=0.0,G=0.0,B=0.0,A=1.0)\">%s</>\n\n"), **TitleStringPtr);
+					RichText += FString::Printf(TEXT("<TextStyle Style=\"Tutorials.Content.HeaderText2\">%s</>\n\n"), **TitleStringPtr);
 				}
 				RichText += Excerpt.RichText;
 
diff --git a/Engine/Source/Editor/IntroTutorials/Private/STutorialContent.cpp b/Engine/Source/Editor/IntroTutorials/Private/STutorialContent.cpp
--- a/Engine/Source/Editor/IntroTutorials/Private/STutorialContent.cpp
+++ b/Engine/Source/Editor/IntroTutorials/Private/STutorialContent.cpp
@@ -176,9 +176,9 @@ TSharedRef<SWidget> STutorialContent::GenerateContentWidget(const FTutorialConte
 		.ContentStyle(TEXT("Tutorials.Content.Text"))
 		.BoldContentStyle(TEXT("Tutorials.Content.TextBold"))
 		.NumberedContentStyle(TEXT("Tutorials.Content.Text"))
-		.Header1Style(TEXT("Tutorials.Content.Header1"))
-		.Header2Style(TEXT("Tutorials.Content.Header2"))
-		.HyperlinkButtonStyle(TEXT("Tutorials.Content.Hyperlink"))
+		.Header1Style(TEXT("Tutorials.Content.HeaderText1"))
+		.Header2Style(TEXT("Tutorials.Content.HeaderText2"))
+		.HyperlinkStyle(TEXT("Tutorials.Content.Hyperlink"))
 		.HyperlinkTextStyle(TEXT("Tutorials.Content.HyperlinkText"))
 		.SeparatorStyle(TEXT("Tutorials.Separator"));
 
diff --git a/Engine/Source/Editor/IntroTutorials/Private/STutorialEditableText.cpp b/Engine/Source/Editor/IntroTutorials/Private/STutorialEditableText.cpp
--- a/Engine/Source/Editor/IntroTutorials/Private/STutorialEditableText.cpp
+++ b/Engine/Source/Editor/IntroTutorials/Private/STutorialEditableText.cpp
@@ -8,6 +8,7 @@
 #include "SColorPicker.h"
 #include "DesktopPlatformModule.h"
 #include "RichTextLayoutMarshaller.h"
+#include "Editor/MainFrame/Public/MainFrame.h"
 
 #define LOCTEXT_NAMESPACE "STutorialEditableText"
 
@@ -34,21 +35,33 @@ void OnDocLinkClicked(const FSlateHyperlinkRun::FMetadata& Metadata)
 	}
 }
 
+static void ParseTutorialLink(const FString &InternalLink)
+{
+	UBlueprint* Blueprint = LoadObject<UBlueprint>(nullptr, *InternalLink);
+	if (Blueprint && Blueprint->GeneratedClass)
+	{
+		FIntroTutorials& IntroTutorials = FModuleManager::GetModuleChecked<FIntroTutorials>(TEXT("IntroTutorials"));
+		IMainFrameModule& MainFrameModule = FModuleManager::LoadModuleChecked<IMainFrameModule>(TEXT("MainFrame"));
+		const bool bRestart = true;
+		IntroTutorials.LaunchTutorial(Blueprint->GeneratedClass->GetDefaultObject<UEditorTutorial>(), bRestart, MainFrameModule.GetParentWindow());
+	}
+}
+
 void OnTutorialLinkClicked(const FSlateHyperlinkRun::FMetadata& Metadata)
 {
 	const FString* Url = Metadata.Find(TEXT("href"));
 	if(Url)
 	{
-	//	ParseTutorialLink(*Url);
+		ParseTutorialLink(*Url);
 	}
 }
 
 static void ParseCodeLink(const FString &InternalLink)
 {
 	// Tokens used by the code parsing. Details in the parse section	
-	static const FString ProjectSpecifier(TEXT("[PROJECT]"));
-	static const FString ProjectRoot(TEXT("[PROJECT]/Source/[PROJECT]/"));
-	static const FString ProjectSuffix(TEXT(".uproject"));
+	static const FString ProjectSpecifier(TEXT("[PROJECTPATH]"));
+	static const FString ProjectPathSpecifier(TEXT("[PROJECT]"));
+	static const FString EnginePathSpecifier(TEXT("[ENGINEPATH]"));
 
 	FString Path;
 	int32 Line = 0;
@@ -73,36 +86,27 @@ static void ParseCodeLink(const FString &InternalLink)
 	ISourceCodeAccessModule& SourceCodeAccessModule = FModuleManager::LoadModuleChecked<ISourceCodeAccessModule>("SourceCodeAccess");
 	ISourceCodeAccessor& SourceCodeAccessor = SourceCodeAccessModule.GetAccessor();
 
-	// If we specified generic project specified as the project name try to replace the name with the name of this project
-	if (InternalLink.Contains(ProjectSpecifier) == true)
+	if (Path.Contains(EnginePathSpecifier) == true)
 	{
-		FString ProjectName = TEXT("Marble");
-		// Try to extract the name of the project
-		FString ProjectPath = FPaths::GetProjectFilePath();
-		if (ProjectPath.EndsWith(ProjectSuffix))
-		{
-			int32 ProjectPathEndIndex;
-			if (ProjectPath.FindLastChar(TEXT('/'), ProjectPathEndIndex) == true)
-			{
-				ProjectName = ProjectPath.RightChop(ProjectPathEndIndex + 1);
-				ProjectName.RemoveFromEnd(*ProjectSuffix);
-			}
-		}
-		// Replace the root path with the name of this project
-		FString RebuiltPath = ProjectRoot + Path;
-		RebuiltPath.ReplaceInline(*ProjectSpecifier, *ProjectName);
-		Path = RebuiltPath;
+		// replace engine path specifier with path to engine
+		Path.ReplaceInline(*EnginePathSpecifier, *FPaths::EngineDir());
+	}
+
+	if (Path.Contains(ProjectSpecifier) == true)
+	{
+		// replace project specifier with path to project
+		Path.ReplaceInline(*ProjectSpecifier, FApp::GetGameName());
 	}
 
-	// Finally create the complete path - project name and all
-	int32 PathEndIndex;
-	FString SolutionPath;
-	if( FDesktopPlatformModule::Get()->GetSolutionPath( SolutionPath ) && SolutionPath.FindLastChar(TEXT('/'), PathEndIndex) == true)
+	if (Path.Contains(ProjectPathSpecifier) == true)
 	{
-		SolutionPath = SolutionPath.LeftChop(SolutionPath.Len() - PathEndIndex - 1);
-		SolutionPath += Path;
-		SourceCodeAccessor.OpenFileAtLine(SolutionPath, Line, Col);
+		// replace project specifier with path to project
+		Path.ReplaceInline(*ProjectPathSpecifier, *FPaths::GetProjectFilePath());
 	}
+
+	Path = FPaths::ConvertRelativePathToFull(Path);
+
+	SourceCodeAccessor.OpenFileAtLine(Path, Line, Col);
 }
 
 void OnCodeLinkClicked(const FSlateHyperlinkRun::FMetadata& Metadata)
@@ -114,22 +118,22 @@ void OnCodeLinkClicked(const FSlateHyperlinkRun::FMetadata& Metadata)
 	}
 }
 
-static void ParseAssetLink(const FString& InternalLink, const FString& Action)
+static void ParseAssetLink(const FString& InternalLink, const FString* Action)
 {
 	UObject* RequiredObject = FindObject<UObject>(ANY_PACKAGE, *InternalLink);
 	if (RequiredObject != nullptr)
 	{
-		if (Action == TEXT("edit"))
-		{
-			FAssetEditorManager::Get().OpenEditorForAsset(RequiredObject);
-		}
-		else if(Action == TEXT("select"))
+		if(Action && *Action == TEXT("select"))
 		{
 			FContentBrowserModule& ContentBrowserModule = FModuleManager::Get().LoadModuleChecked<FContentBrowserModule>("ContentBrowser");
 			TArray<UObject*> AssetToBrowse;
 			AssetToBrowse.Add(RequiredObject);
 			ContentBrowserModule.Get().SyncBrowserToAssets(AssetToBrowse);
 		}
+		else
+		{
+			FAssetEditorManager::Get().OpenEditorForAsset(RequiredObject);
+		}
 	}
 }
 
@@ -137,9 +141,9 @@ void OnAssetLinkClicked(const FSlateHyperlinkRun::FMetadata& Metadata)
 {
 	const FString* Url = Metadata.Find(TEXT("href"));
 	const FString* Action = Metadata.Find(TEXT("action"));
-	if(Url && Action)
+	if(Url)
 	{
-		ParseAssetLink(*Url, *Action);
+		ParseAssetLink(*Url, Action);
 	}
 }
 
@@ -151,7 +155,6 @@ void STutorialEditableText::Construct(const FArguments& InArgs)
 	OnTextChanged = InArgs._OnTextChanged;
 	OnTextCommitted = InArgs._OnTextCommitted;
 
-	CurrentHyperlinkType = EHyperlinkType::Browser;
 	bNewHyperlink = true;
 
 	// Setup text styles
@@ -170,16 +173,48 @@ void STutorialEditableText::Construct(const FArguments& InArgs)
 		FEditorStyle::Get().GetWidgetStyle<FTextBlockStyle>("Tutorials.Content.Text")
 		);
 
-	OnBrowserLinkClicked = FSlateHyperlinkRun::FOnClick::CreateStatic(&TutorialTextHelpers::OnBrowserLinkClicked);
-	OnDocLinkClicked = FSlateHyperlinkRun::FOnClick::CreateStatic(&TutorialTextHelpers::OnDocLinkClicked);
-	OnTutorialLinkClicked = FSlateHyperlinkRun::FOnClick::CreateStatic(&TutorialTextHelpers::OnTutorialLinkClicked);
-	OnCodeLinkClicked = FSlateHyperlinkRun::FOnClick::CreateStatic(&TutorialTextHelpers::OnCodeLinkClicked);
-	OnAssetLinkClicked = FSlateHyperlinkRun::FOnClick::CreateStatic(&TutorialTextHelpers::OnAssetLinkClicked);
-	RichTextMarshaller->AppendInlineDecorator(FHyperlinkDecorator::Create(TEXT("browser"), OnBrowserLinkClicked));
-	RichTextMarshaller->AppendInlineDecorator(FHyperlinkDecorator::Create(TEXT("udn"), OnDocLinkClicked));
-	RichTextMarshaller->AppendInlineDecorator(FHyperlinkDecorator::Create(TEXT("tutorial"), OnTutorialLinkClicked));
-	RichTextMarshaller->AppendInlineDecorator(FHyperlinkDecorator::Create(TEXT("code"), OnCodeLinkClicked));
-	RichTextMarshaller->AppendInlineDecorator(FHyperlinkDecorator::Create(TEXT("asset"), OnAssetLinkClicked));
+	HyperlinkDescs.Add(MakeShareable(new FHyperlinkTypeDesc(
+		EHyperlinkType::Browser, 
+		LOCTEXT("BrowserLinkTypeLabel", "URL"), 
+		LOCTEXT("BrowserLinkTypeTooltip", "A link that opens a browser window (e.g. http://www.unrealengine.com)"),
+		TEXT("browser"),
+		FSlateHyperlinkRun::FOnClick::CreateStatic(&TutorialTextHelpers::OnBrowserLinkClicked))));
+
+	CurrentHyperlinkType = HyperlinkDescs[0];
+
+	HyperlinkDescs.Add(MakeShareable(new FHyperlinkTypeDesc(
+		EHyperlinkType::UDN, 
+		LOCTEXT("UDNLinkTypeLabel", "UDN"), 
+		LOCTEXT("UDNLinkTypeTooltip", "A link that opens some UDN documentation (e.g. /Engine/Blueprints/UserGuide/Types/ClassBlueprint)"),
+		TEXT("udn"),
+		FSlateHyperlinkRun::FOnClick::CreateStatic(&TutorialTextHelpers::OnDocLinkClicked))));
+
+	HyperlinkDescs.Add(MakeShareable(new FHyperlinkTypeDesc(
+		EHyperlinkType::Asset, 
+		LOCTEXT("AssetLinkTypeLabel", "Asset"), 
+		LOCTEXT("AssetLinkTypeTooltip", "A link that opens an asset (e.g. /Game/StaticMeshes/SphereMesh.SphereMesh)"),
+		TEXT("asset"),
+		FSlateHyperlinkRun::FOnClick::CreateStatic(&TutorialTextHelpers::OnAssetLinkClicked))));
+
+	HyperlinkDescs.Add(MakeShareable(new FHyperlinkTypeDesc(
+		EHyperlinkType::Code, 
+		LOCTEXT("CodeLinkTypeLabel", "Code"), 
+		LOCTEXT("CodeLinkTypeTooltip", "A link that opens code in your selected IDE.\nFor example: [PROJECTPATH]/Private/SourceFile.cpp,1,1.\nThe numbers correspond to line number and column number.\nYou can use [PROJECT], [PROJECTPATH] and [ENGINEPATH] tags to make paths."),
+		TEXT("code"),
+		FSlateHyperlinkRun::FOnClick::CreateStatic(&TutorialTextHelpers::OnCodeLinkClicked))));
+
+	HyperlinkDescs.Add(MakeShareable(new FHyperlinkTypeDesc(
+		EHyperlinkType::Tutorial, 
+		LOCTEXT("TutorialLinkTypeLabel", "Tutorial"), 
+		LOCTEXT("TutorialLinkTypeTooltip", "A type of asset link that opens another tutorial, e.g. /Game/Tutorials/StaticMeshTutorial.StaticMeshTutorial"),
+		TEXT("tutorial"),
+		FSlateHyperlinkRun::FOnClick::CreateStatic(&TutorialTextHelpers::OnTutorialLinkClicked))));
+
+	for(const auto& HyperlinkDesc : HyperlinkDescs)
+	{
+		RichTextMarshaller->AppendInlineDecorator(FHyperlinkDecorator::Create(HyperlinkDesc->Id, HyperlinkDesc->OnClickedDelegate));
+	}
+
 	RichTextMarshaller->AppendInlineDecorator(FTextStyleDecorator::Create());
 
 	this->ChildSlot
@@ -256,6 +291,7 @@ void STutorialEditableText::Construct(const FArguments& InArgs)
 
 						+SGridPanel::Slot(0, 0)
 						.HAlign(HAlign_Right)
+						.VAlign(VAlign_Center)
 						.Padding(FMargin(2.0f))
 						[
 							SNew(STextBlock)
@@ -275,6 +311,7 @@ void STutorialEditableText::Construct(const FArguments& InArgs)
 
 						+SGridPanel::Slot(0, 1)
 						.HAlign(HAlign_Right)
+						.VAlign(VAlign_Center)
 						.Padding(FMargin(2.0f))
 						[
 							SNew(STextBlock)
@@ -293,7 +330,17 @@ void STutorialEditableText::Construct(const FArguments& InArgs)
 
 						+SGridPanel::Slot(0, 2)
 						.HAlign(HAlign_Right)
+						.VAlign(VAlign_Center)
 						.Padding(FMargin(2.0f))
+						[
+							SNew(STextBlock)
+							.TextStyle(FEditorStyle::Get(), "TutorialEditableText.Toolbar.Text")
+							.Text(LOCTEXT("HyperlinkTypeLabel", "Type:"))
+						]
+
+						+SGridPanel::Slot(1, 2)
+						.Padding(FMargin(2.0f))
+						.VAlign(VAlign_Center)
 						.ColumnSpan(2)
 						[
 							SNew(SHorizontalBox)
@@ -301,72 +348,46 @@ void STutorialEditableText::Construct(const FArguments& InArgs)
 							.AutoWidth()
 							.VAlign(VAlign_Center)
 							[
-								SNew(SCheckBox)
-								.Style(FEditorStyle::Get(), "RadioButton")
-								.Type(ESlateCheckBoxType::CheckBox)
-								.IsChecked(this, &STutorialEditableText::IsCreatingBrowserLink)
-								.OnCheckStateChanged(this, &STutorialEditableText::OnCheckBrowserLink)
-								[
-									SNew(STextBlock)
-									.TextStyle(FEditorStyle::Get(), "TutorialEditableText.Toolbar.Text")
-									.Text(LOCTEXT("BrowserLinkTypeLabel", "URL"))
-									.ToolTipText(LOCTEXT("BrowserLinkTypeTooltip", "A link that opens a browser window (e.g. http://www.unrealengine.com)"))
-								]
-							]
-							+SHorizontalBox::Slot()
-							.AutoWidth()
-							.VAlign(VAlign_Center)
-							[
-								SNew(SCheckBox)
-								.IsEnabled(false)
-								.Style(FEditorStyle::Get(), "RadioButton")
-								.Type(ESlateCheckBoxType::CheckBox)
-								.IsChecked(this, &STutorialEditableText::IsCreatingUDNLink)
-								.OnCheckStateChanged(this, &STutorialEditableText::OnCheckUDNLink)
+								SNew(SComboBox<TSharedPtr<FHyperlinkTypeDesc>>)
+								.OptionsSource(&HyperlinkDescs)
+								.ComboBoxStyle(FEditorStyle::Get(), "TutorialEditableText.Toolbar.ComboBox")
+								.OnSelectionChanged(this, &STutorialEditableText::OnActiveHyperlinkChanged)
+								.OnGenerateWidget(this, &STutorialEditableText::GenerateHyperlinkComboEntry)
+								.ContentPadding(0.0f)
+								.InitiallySelectedItem(HyperlinkDescs[0])
+								.Content()
 								[
-									SNew(STextBlock)
-									.TextStyle(FEditorStyle::Get(), "TutorialEditableText.Toolbar.Text")
-									.Text(LOCTEXT("UDNLinkTypeLabel", "UDN"))
-									.ToolTipText(LOCTEXT("UDNLinkTypeTooltip", "A link that opens some UDN documentation (e.g. /Engine/Blueprints/UserGuide/Types/ClassBlueprint)"))
+									SNew(SBox)
+									.Padding(FMargin(0.0f, 0.0f, 2.0f, 0.0f))
+									.MinDesiredWidth(100.0f)
+									[
+										SNew(STextBlock)
+										.TextStyle(FEditorStyle::Get(), "TutorialEditableText.Toolbar.Text")
+										.Text(this, &STutorialEditableText::GetActiveHyperlinkName)
+										.ToolTipText(this, &STutorialEditableText::GetActiveHyperlinkTooltip)
+									]
 								]
 							]
 							+SHorizontalBox::Slot()
 							.AutoWidth()
 							.VAlign(VAlign_Center)
+							.Padding(5.0f, 0.0f, 0.0f, 0.0f)
 							[
 								SNew(SCheckBox)
-								.IsEnabled(false)
-								.Style(FEditorStyle::Get(), "RadioButton")
-								.Type(ESlateCheckBoxType::CheckBox)
-								.IsChecked(this, &STutorialEditableText::IsCreatingAssetLink)
-								.OnCheckStateChanged(this, &STutorialEditableText::OnCheckAssetLink)
+								.ToolTipText(LOCTEXT("OpenAssetTooltip", "Should this link open the asset or just highlight it in the content browser?"))
+								.Visibility(this, &STutorialEditableText::GetOpenAssetVisibility)
+								.IsChecked(this, &STutorialEditableText::IsOpenAssetChecked)
+								.OnCheckStateChanged(this, &STutorialEditableText::HandleOpenAssetCheckStateChanged)
 								[
 									SNew(STextBlock)
 									.TextStyle(FEditorStyle::Get(), "TutorialEditableText.Toolbar.Text")
-									.Text(LOCTEXT("AssetLinkTypeLabel", "Asset"))
-									.ToolTipText(LOCTEXT("AssetLinkTypeTooltip", "A link that opens an asset (e.g. Game/StaticMeshes/SphereMesh)"))
+									.Text(LOCTEXT("OpenAssetLabel", "Open Asset"))
 								]
 							]
 							+SHorizontalBox::Slot()
-							.AutoWidth()
+							.FillWidth(1.0f)
 							.VAlign(VAlign_Center)
-							[
-								SNew(SCheckBox)
-								.IsEnabled(false)
-								.Style(FEditorStyle::Get(), "RadioButton")
-								.Type(ESlateCheckBoxType::CheckBox)
-								.IsChecked(this, &STutorialEditableText::IsCreatingCodeLink)
-								.OnCheckStateChanged(this, &STutorialEditableText::OnCheckCodeLink)
-								[
-									SNew(STextBlock)
-									.TextStyle(FEditorStyle::Get(), "TutorialEditableText.Toolbar.Text")
-									.Text(LOCTEXT("CodeLinkTypeLabel", "Code"))
-									.ToolTipText(LOCTEXT("CodeLinkTypeTooltip", "A link that opens code (e.g. Private/SourceFile.cpp,1,1)"))
-								]
-							]
-							+SHorizontalBox::Slot()
-							.AutoWidth()
-							.Padding(FMargin(5.0f, 0.0f, 0.0f, 0.0f))
+							.HAlign(HAlign_Right)
 							[
 								SNew(SButton)
 								.ButtonStyle(FEditorStyle::Get(), "TutorialEditableText.Toolbar.Button")
@@ -526,6 +547,19 @@ void STutorialEditableText::StyleSelectedText()
 	}
 }
 
+TSharedPtr<FHyperlinkTypeDesc> STutorialEditableText::GetHyperlinkTypeFromId(const FString& InId) const
+{
+	TSharedPtr<FHyperlinkTypeDesc> FoundType;
+	for(const auto& Desc : HyperlinkDescs)
+	{
+		if(Desc->Id == InId)
+		{
+			return Desc;
+		}
+	}
+
+	return FoundType;
+}
 
 void STutorialEditableText::HandleHyperlinkComboOpened()
 {
@@ -552,6 +586,10 @@ void STutorialEditableText::HandleHyperlinkComboOpened()
 	{
 		const FString* const URLUnderCursor = Run->GetRunInfo().MetaData.Find(TEXT("href"));
 		HyperlinkURLTextBox->SetText((URLUnderCursor) ? FText::FromString(*URLUnderCursor) : FText());
+
+		const FString* const IdUnderCursor = Run->GetRunInfo().MetaData.Find(TEXT("id"));
+		CurrentHyperlinkType = IdUnderCursor ? GetHyperlinkTypeFromId(*IdUnderCursor) : HyperlinkDescs[0];
+
 		FString RunText;
 		Run->AppendTextTo(RunText);
 		HyperlinkNameTextBlock->SetText(FText::FromString(RunText));
@@ -574,16 +612,21 @@ FReply STutorialEditableText::HandleInsertHyperLinkClicked()
 
 		// Create the correct meta-information for this run, so that valid source rich-text formatting can be generated for it
 		FRunInfo RunInfo(TEXT("a"));
-		RunInfo.MetaData.Add(TEXT("id"), TEXT("browser"));
+		RunInfo.MetaData.Add(TEXT("id"), CurrentHyperlinkType->Id);
 		RunInfo.MetaData.Add(TEXT("href"), URL.ToString());
 		RunInfo.MetaData.Add(TEXT("style"), TEXT("Tutorials.Content.Hyperlink"));
 
+		if(CurrentHyperlinkType->Type == EHyperlinkType::Asset)
+		{
+			RunInfo.MetaData.Add(TEXT("action"), bOpenAsset ? TEXT("edit") : TEXT("select"));
+		}
+
 		// Create the new run, and then insert it at the cursor position
 		TSharedRef<FSlateHyperlinkRun> HyperlinkRun = FSlateHyperlinkRun::Create(
 			RunInfo, 
 			MakeShareable(new FString(Name.ToString())), 
 			FEditorStyle::Get().GetWidgetStyle<FHyperlinkStyle>(FName(TEXT("Tutorials.Content.Hyperlink"))), 
-			OnBrowserLinkClicked
+			CurrentHyperlinkType->OnClickedDelegate
 			);
 
 		// @todo: if the rich editable text box allowed us to replace a run that we found under the cursor (or returned a non-const
@@ -600,74 +643,62 @@ EVisibility STutorialEditableText::GetToolbarVisibility() const
 	return FontComboBox->IsOpen() || HyperlinkComboButton->IsOpen() || HasKeyboardFocus() || HasFocusedDescendants() ? EVisibility::Visible : EVisibility::Collapsed;
 }
 
-ESlateCheckBoxState::Type STutorialEditableText::IsCreatingBrowserLink() const
+FText STutorialEditableText::GetHyperlinkButtonText() const
 {
-	return CurrentHyperlinkType == EHyperlinkType::Browser ? ESlateCheckBoxState::Checked : ESlateCheckBoxState::Unchecked;
+	return bNewHyperlink ? LOCTEXT("HyperlinkInsertLabel", "Insert Hyperlink") : LOCTEXT("HyperlinkSetLabel", "Set Hyperlink");
 }
 
-void STutorialEditableText::OnCheckBrowserLink(ESlateCheckBoxState::Type State)
+void STutorialEditableText::OnActiveHyperlinkChanged(TSharedPtr<FHyperlinkTypeDesc> NewValue, ESelectInfo::Type SelectionType)
 {
-	if(State == ESlateCheckBoxState::Checked)
-	{
-		CurrentHyperlinkType = EHyperlinkType::Browser;
-	}
+	CurrentHyperlinkType = NewValue;
 }
 
-ESlateCheckBoxState::Type STutorialEditableText::IsCreatingUDNLink() const
+TSharedRef<SWidget> STutorialEditableText::GenerateHyperlinkComboEntry(TSharedPtr<FHyperlinkTypeDesc> SourceEntry)
 {
-	return CurrentHyperlinkType == EHyperlinkType::UDN ? ESlateCheckBoxState::Checked : ESlateCheckBoxState::Unchecked;
+	return SNew(SBorder)
+		.BorderImage( FCoreStyle::Get().GetBrush( "NoBorder" ) )
+		.ForegroundColor(FCoreStyle::Get().GetSlateColor("InvertedForeground"))
+		[
+			SNew(STextBlock)
+			.Text(SourceEntry->Text)
+			.ToolTipText(SourceEntry->TooltipText)
+			.TextStyle(&FEditorStyle::Get().GetWidgetStyle<FTextBlockStyle>("TutorialEditableText.Toolbar.Text"))
+		];
 }
 
-void STutorialEditableText::OnCheckUDNLink(ESlateCheckBoxState::Type State)
+FText STutorialEditableText::GetActiveHyperlinkName() const
 {
-	if(State == ESlateCheckBoxState::Checked)
+	if(CurrentHyperlinkType.IsValid())
 	{
-		CurrentHyperlinkType = EHyperlinkType::UDN;
+		return CurrentHyperlinkType->Text;
 	}
-}
 
-ESlateCheckBoxState::Type STutorialEditableText::IsCreatingTutorialLink() const
-{
-	return CurrentHyperlinkType == EHyperlinkType::Tutorial ? ESlateCheckBoxState::Checked : ESlateCheckBoxState::Unchecked;
+	return FText();
 }
 
-void STutorialEditableText::OnCheckTutorialLink(ESlateCheckBoxState::Type State)
+FText STutorialEditableText::GetActiveHyperlinkTooltip() const
 {
-	if(State == ESlateCheckBoxState::Checked)
+	if(CurrentHyperlinkType.IsValid())
 	{
-		CurrentHyperlinkType = EHyperlinkType::Tutorial;
+		return CurrentHyperlinkType->TooltipText;
 	}
-}
 
-ESlateCheckBoxState::Type STutorialEditableText::IsCreatingCodeLink() const
-{
-	return CurrentHyperlinkType == EHyperlinkType::Code ? ESlateCheckBoxState::Checked : ESlateCheckBoxState::Unchecked;
+	return FText();
 }
 
-void STutorialEditableText::OnCheckCodeLink(ESlateCheckBoxState::Type State)
+EVisibility STutorialEditableText::GetOpenAssetVisibility() const
 {
-	if(State == ESlateCheckBoxState::Checked)
-	{
-		CurrentHyperlinkType = EHyperlinkType::Code;
-	}
+	return CurrentHyperlinkType.IsValid() && CurrentHyperlinkType->Type == EHyperlinkType::Asset ? EVisibility::Visible : EVisibility::Collapsed;
 }
 
-ESlateCheckBoxState::Type STutorialEditableText::IsCreatingAssetLink() const
+void STutorialEditableText::HandleOpenAssetCheckStateChanged(ESlateCheckBoxState::Type InCheckState)
 {
-	return CurrentHyperlinkType == EHyperlinkType::Asset ? ESlateCheckBoxState::Checked : ESlateCheckBoxState::Unchecked;
+	bOpenAsset = !bOpenAsset;
 }
 
-void STutorialEditableText::OnCheckAssetLink(ESlateCheckBoxState::Type State)
+ESlateCheckBoxState::Type STutorialEditableText::IsOpenAssetChecked() const
 {
-	if(State == ESlateCheckBoxState::Checked)
-	{
-		CurrentHyperlinkType = EHyperlinkType::Asset;
-	}
-}
-
-FText STutorialEditableText::GetHyperlinkButtonText() const
-{
-	return bNewHyperlink ? LOCTEXT("HyperlinkInsertLabel", "Insert Hyperlink") : LOCTEXT("HyperlinkSetLabel", "Set Hyperlink");
+	return bOpenAsset ? ESlateCheckBoxState::Checked : ESlateCheckBoxState::Unchecked;
 }
 
 #undef LOCTEXT_NAMESPACE
\ No newline at end of file
diff --git a/Engine/Source/Editor/IntroTutorials/Private/STutorialEditableText.h b/Engine/Source/Editor/IntroTutorials/Private/STutorialEditableText.h
--- a/Engine/Source/Editor/IntroTutorials/Private/STutorialEditableText.h
+++ b/Engine/Source/Editor/IntroTutorials/Private/STutorialEditableText.h
@@ -177,6 +177,34 @@ public:
 	}
 };
 
+/** Helper struct to hold info about hyperlink types */
+struct FHyperlinkTypeDesc
+{
+	FHyperlinkTypeDesc(EHyperlinkType::Type InType, const FText& InText, const FText& InTooltipText, const FString& InId, FSlateHyperlinkRun::FOnClick InOnClickedDelegate)
+		: Type(InType)
+		, Text(InText)
+		, TooltipText(InTooltipText)
+		, Id(InId)
+		, OnClickedDelegate(InOnClickedDelegate)
+	{
+	}
+
+	/** The type of the link */
+	EHyperlinkType::Type Type;
+
+	/** Tag used by this hyperlink's run */
+	FString Id;
+
+	/** Text to display in the UI */
+	FText Text;
+
+	/** Tooltip text to display in the UI */
+	FText TooltipText;
+
+	/** Delegate to execute for this hyperlink's run */
+	FSlateHyperlinkRun::FOnClick OnClickedDelegate;
+};
+
 
 class STutorialEditableText : public SCompoundWidget
 {
@@ -226,27 +254,23 @@ protected:
 
 	EVisibility GetToolbarVisibility() const;
 
-	ESlateCheckBoxState::Type IsCreatingBrowserLink() const;
-
-	void OnCheckBrowserLink(ESlateCheckBoxState::Type State);
-
-	ESlateCheckBoxState::Type IsCreatingUDNLink() const;
+	FText GetHyperlinkButtonText() const;
 
-	void OnCheckUDNLink(ESlateCheckBoxState::Type State);
+	void OnActiveHyperlinkChanged(TSharedPtr<FHyperlinkTypeDesc> NewValue, ESelectInfo::Type SelectionType);
 
-	ESlateCheckBoxState::Type IsCreatingTutorialLink() const;
+	TSharedRef<SWidget> GenerateHyperlinkComboEntry(TSharedPtr<FHyperlinkTypeDesc> SourceEntry);
 
-	void OnCheckTutorialLink(ESlateCheckBoxState::Type State);
+	FText GetActiveHyperlinkName() const;
 
-	ESlateCheckBoxState::Type IsCreatingCodeLink() const;
+	FText GetActiveHyperlinkTooltip() const;
 
-	void OnCheckCodeLink(ESlateCheckBoxState::Type State);
+	TSharedPtr<FHyperlinkTypeDesc> GetHyperlinkTypeFromId(const FString& InId) const;
 
-	ESlateCheckBoxState::Type IsCreatingAssetLink() const;
+	EVisibility GetOpenAssetVisibility() const;
 
-	void OnCheckAssetLink(ESlateCheckBoxState::Type State);
+	void HandleOpenAssetCheckStateChanged(ESlateCheckBoxState::Type InCheckState);
 
-	FText GetHyperlinkButtonText() const;
+	ESlateCheckBoxState::Type IsOpenAssetChecked() const;
 
 protected:
 	TSharedPtr<SMultiLineEditableTextBox> RichEditableTextBox;
@@ -270,7 +294,10 @@ protected:
 	FOnTextCommitted OnTextCommitted;
 	FOnTextChanged OnTextChanged;
 
-	EHyperlinkType::Type CurrentHyperlinkType;
+	TArray<TSharedPtr<FHyperlinkTypeDesc>> HyperlinkDescs;
+	TSharedPtr<FHyperlinkTypeDesc> CurrentHyperlinkType;
+
+	bool bOpenAsset;
 
 	bool bNewHyperlink;
 };
diff --git a/Engine/Source/Editor/IntroTutorials/Public/EditorTutorial.h b/Engine/Source/Editor/IntroTutorials/Public/EditorTutorial.h
--- a/Engine/Source/Editor/IntroTutorials/Public/EditorTutorial.h
+++ b/Engine/Source/Editor/IntroTutorials/Public/EditorTutorial.h
@@ -72,7 +72,7 @@ struct INTROTUTORIALS_API FTutorialContent
 
 	FTutorialContent()
 	{
-		Type = ETutorialContent::Text;
+		Type = ETutorialContent::RichText;
 	}
 
 	/** The type of this content */
