Fix non-unity in STutorialEditableText.cpp

diff --git a/Engine/Source/Editor/IntroTutorials/Private/STutorialEditableText.cpp b/Engine/Source/Editor/IntroTutorials/Private/STutorialEditableText.cpp
--- a/Engine/Source/Editor/IntroTutorials/Private/STutorialEditableText.cpp
+++ b/Engine/Source/Editor/IntroTutorials/Private/STutorialEditableText.cpp
@@ -8,6 +8,7 @@
 #include "SColorPicker.h"
 #include "DesktopPlatformModule.h"
 #include "RichTextLayoutMarshaller.h"
+#include "EditorTutorial.h"
 #include "Editor/MainFrame/Public/MainFrame.h"
 
 #define LOCTEXT_NAMESPACE "STutorialEditableText"
