Added intro animation to tutorials

Intro 'bounce' is driven by a float curve asset.
Because of rich text not working very well when zoomed at the moment, only the border is zoomed when animating. The content has the inverse zoom applied to it to counter this.

diff --git a/Engine/Source/Editor/IntroTutorials/Private/IntroTutorials.cpp b/Engine/Source/Editor/IntroTutorials/Private/IntroTutorials.cpp
--- a/Engine/Source/Editor/IntroTutorials/Private/IntroTutorials.cpp
+++ b/Engine/Source/Editor/IntroTutorials/Private/IntroTutorials.cpp
@@ -56,6 +56,7 @@ const FWelcomeTutorialProperties FIntroTutorials::TemplateOverview(TemplateOverv
 
 FIntroTutorials::FIntroTutorials()
 	: CurrentObjectClass(nullptr)
+	, ContentIntroCurve(nullptr)
 {
 	bDisableTutorials = false;
 	bEnablePostTutorialSurveys = false;
@@ -225,6 +226,12 @@ void FIntroTutorials::StartupModule()
 	PropertyEditorModule.RegisterCustomPropertyTypeLayout("TutorialContent", FOnGetPropertyTypeCustomizationInstance::CreateStatic(&FTutorialContentCustomization::MakeInstance));
 	PropertyEditorModule.RegisterCustomPropertyTypeLayout("TutorialContentAnchor", FOnGetPropertyTypeCustomizationInstance::CreateStatic(&FTutorialContentAnchorCustomization::MakeInstance));
 	PropertyEditorModule.RegisterCustomClassLayout("EditorTutorial", FOnGetDetailCustomizationInstance::CreateStatic(&FEditorTutorialDetailsCustomization::MakeInstance));
+
+	ContentIntroCurve = LoadObject<UCurveFloat>(nullptr, TEXT("/Engine/Tutorial/ContentIntroCurve.ContentIntroCurve"));
+	if(ContentIntroCurve)
+	{
+		ContentIntroCurve->AddToRoot();
+	}
 }
 
 void FIntroTutorials::ShutdownModule()
@@ -266,6 +273,12 @@ void FIntroTutorials::ShutdownModule()
 		PropertyEditorModule.UnregisterCustomPropertyTypeLayout("TutorialWidgetReference");
 		PropertyEditorModule.UnregisterCustomClassLayout("EditorTutorial");
 	}
+
+	if(ContentIntroCurve)
+	{
+		ContentIntroCurve->RemoveFromRoot();
+		ContentIntroCurve = nullptr;
+	}
 }
 
 void FIntroTutorials::AddSummonTutorialsMenuExtension(FMenuBuilder& MenuBuilder)
@@ -794,4 +807,14 @@ TSharedRef<SWidget> FIntroTutorials::CreateTutorialsWidget(FName InContext, TWea
 		.ContextWindow(InContextWindow);
 }
 
+float FIntroTutorials::GetIntroCurveValue(float InTime)
+{
+	if(ContentIntroCurve)
+	{
+		return ContentIntroCurve->GetFloatValue(InTime);
+	}
+
+	return 1.0f;
+}
+
 #undef LOCTEXT_NAMESPACE
\ No newline at end of file
diff --git a/Engine/Source/Editor/IntroTutorials/Private/IntroTutorials.h b/Engine/Source/Editor/IntroTutorials/Private/IntroTutorials.h
--- a/Engine/Source/Editor/IntroTutorials/Private/IntroTutorials.h
+++ b/Engine/Source/Editor/IntroTutorials/Private/IntroTutorials.h
@@ -73,6 +73,8 @@ public:
 
 	void GoToNextStage(TWeakPtr<SWindow> InNavigationWindow);
 
+	float GetIntroCurveValue(float InTime);
+
 private:
 
 	/** IModuleInterface implementation */
@@ -197,4 +199,7 @@ private:
 
 	/** Root widget for tutorial overlay system */
 	TSharedPtr<STutorialRoot> TutorialRoot;
+
+	/** Curve asset for intros */
+	UCurveFloat* ContentIntroCurve;
 };
\ No newline at end of file
diff --git a/Engine/Source/Editor/IntroTutorials/Private/STutorialContent.cpp b/Engine/Source/Editor/IntroTutorials/Private/STutorialContent.cpp
--- a/Engine/Source/Editor/IntroTutorials/Private/STutorialContent.cpp
+++ b/Engine/Source/Editor/IntroTutorials/Private/STutorialContent.cpp
@@ -3,6 +3,7 @@
 #include "IntroTutorialsPrivatePCH.h"
 #include "STutorialContent.h"
 #include "STutorialEditableText.h"
+#include "IntroTutorials.h"
 
 #define LOCTEXT_NAMESPACE "STutorialContent"
 
@@ -10,9 +11,11 @@ namespace TutorialConstants
 {
 	const float BorderPulseAnimationLength = 1.0f;
 	const float BorderIntroAnimationLength = 0.5f;
+	const float ContentIntroAnimationLength = 0.5f;
 	const float MinBorderOpacity = 0.1f;
 	const float ShadowScale = 8.0f;
 	const float MaxBorderOffset = 8.0f;
+	const float BorderSize = 24.0f;
 }
 
 const float ContentOffset = 10.0f;
@@ -33,6 +36,9 @@ void STutorialContent::Construct(const FArguments& InArgs, const FTutorialConten
 	BorderIntroAnimation.Play();
 	BorderPulseAnimation.Play();
 
+	ContentIntroAnimation.AddCurve(0.0f, TutorialConstants::ContentIntroAnimationLength, ECurveEaseFunction::Linear);
+	ContentIntroAnimation.Play();
+
 	if (InContent.Text.IsEmpty() == true)
 	{
 		ChildSlot
@@ -45,37 +51,47 @@ void STutorialContent::Construct(const FArguments& InArgs, const FTutorialConten
 
 	ChildSlot
 	[
-		SAssignNew(ContentWidget, SBorder)
-
-		// Add more padding if the content is to be displayed centrally (i.e. not on a widget)
-		.Padding(24.0f)
-		.BorderImage(FEditorStyle::GetBrush("Tutorials.Border"))
-		.Visibility(this, &STutorialContent::GetVisibility)
-		.BorderBackgroundColor(this, &STutorialContent::GetBackgroundColor)
-		.ForegroundColor(FCoreStyle::Get().GetSlateColor("InvertedForeground"))
+		SNew(SFxWidget)
+		.RenderScale(this, &STutorialContent::GetAnimatedZoom)
+		.RenderScaleOrigin(FVector2D(0.5f, 0.5f))
 		[
-			SNew(SHorizontalBox)
-			+SHorizontalBox::Slot()
-			.AutoWidth()
-			.MaxWidth(600.0f)
-			.VAlign(VAlign_Center)
-			[
-				GenerateContentWidget(InContent, InArgs._WrapTextAt, DocumentationPage)
-			]
-			+SHorizontalBox::Slot()
-			.AutoWidth()
-			.VAlign(VAlign_Top)
-			.Padding(5.0f, 0.0f, 0.0f, 0.0f)
+			SAssignNew(ContentWidget, SBorder)
+
+			// Add more padding if the content is to be displayed centrally (i.e. not on a widget)
+			.Padding(TutorialConstants::BorderSize)
+			.BorderImage(FEditorStyle::GetBrush("Tutorials.Border"))
+			.Visibility(this, &STutorialContent::GetVisibility)
+			.BorderBackgroundColor(this, &STutorialContent::GetBackgroundColor)
+			.ForegroundColor(FCoreStyle::Get().GetSlateColor("InvertedForeground"))
 			[
-				SNew(SButton)
-				.OnClicked(this, &STutorialContent::OnCloseButtonClicked)
-				.Visibility(this, &STutorialContent::GetCloseButtonVisibility)
-				.ButtonStyle(&FEditorStyle::Get().GetWidgetStyle<FButtonStyle>("Tutorials.Browser.Button"))
-				.ContentPadding(0.0f)
+				SNew(SFxWidget)
+				.RenderScale(this, &STutorialContent::GetInverseAnimatedZoom)
+				.RenderScaleOrigin(FVector2D(0.5f, 0.5f))
 				[
-					SNew(SImage)
-					.Image(FEditorStyle::GetBrush("Symbols.X"))
-					.ColorAndOpacity(FLinearColor(0.0f, 0.0f, 0.0f, 1.0f))
+					SNew(SHorizontalBox)
+					+SHorizontalBox::Slot()
+					.AutoWidth()
+					.MaxWidth(600.0f)
+					.VAlign(VAlign_Center)
+					[
+						GenerateContentWidget(InContent, InArgs._WrapTextAt, DocumentationPage)
+					]
+					+SHorizontalBox::Slot()
+					.AutoWidth()
+					.VAlign(VAlign_Top)
+					.Padding(5.0f, 0.0f, 0.0f, 0.0f)
+					[
+						SNew(SButton)
+						.OnClicked(this, &STutorialContent::OnCloseButtonClicked)
+						.Visibility(this, &STutorialContent::GetCloseButtonVisibility)
+						.ButtonStyle(&FEditorStyle::Get().GetWidgetStyle<FButtonStyle>("Tutorials.Browser.Button"))
+						.ContentPadding(0.0f)
+						[
+							SNew(SImage)
+							.Image(FEditorStyle::GetBrush("Symbols.X"))
+							.ColorAndOpacity(FLinearColor(0.0f, 0.0f, 0.0f, 1.0f))
+						]
+					]
 				]
 			]
 		]
@@ -349,4 +365,22 @@ FSlateColor STutorialContent::GetBackgroundColor() const
 	return CachedContentGeometry.IsUnderLocation(CursorPos) ? FLinearColor(1.0f, 1.0f, 1.0f, 1.0f) : FLinearColor(1.0f, 1.0f, 1.0f, 0.8f);
 }
 
+float STutorialContent::GetAnimatedZoom() const
+{
+	if(ContentIntroAnimation.IsPlaying())
+	{
+		FIntroTutorials& IntroTutorials = FModuleManager::GetModuleChecked<FIntroTutorials>(TEXT("IntroTutorials"));
+		return 0.75f + (0.25f * IntroTutorials.GetIntroCurveValue(ContentIntroAnimation.GetLerp()));
+	}
+	else
+	{
+		return 1.0f;
+	}
+}
+
+float STutorialContent::GetInverseAnimatedZoom() const
+{
+	return 1.0f / GetAnimatedZoom();
+}
+
 #undef LOCTEXT_NAMESPACE
\ No newline at end of file
diff --git a/Engine/Source/Editor/IntroTutorials/Private/STutorialContent.h b/Engine/Source/Editor/IntroTutorials/Private/STutorialContent.h
--- a/Engine/Source/Editor/IntroTutorials/Private/STutorialContent.h
+++ b/Engine/Source/Editor/IntroTutorials/Private/STutorialContent.h
@@ -77,6 +77,12 @@ private:
 	/** Alter the background color depending on hover state */
 	FSlateColor GetBackgroundColor() const;
 
+	/** Get zoom level padding for content (animated for intro) */
+	float GetAnimatedZoom() const;
+	
+	/** Get inverse zoom level padding for content - needed because rich text content doesnt scale well */
+	float GetInverseAnimatedZoom() const;
+
 private:
 
 	/** Copy of the window size we were last draw at */
@@ -116,6 +122,9 @@ private:
 	FCurveSequence BorderPulseAnimation;
 	FCurveSequence BorderIntroAnimation;
 
+	/** Animation curve for displaying content */
+	FCurveSequence ContentIntroAnimation;
+
 	/** Documentation page reference to use if we are displaying a UDN doc - we need this otherwise the page will be freed */
 	TSharedPtr<IDocumentationPage> DocumentationPage;
 };
\ No newline at end of file
