Always display the tutorial button. Button action is context sensitive.

Tutorial button launches the tutorial browser if no tutorial exists (or the tutorial has already been taken & completed).

diff --git a/Engine/Source/Editor/IntroTutorials/Private/EditorTutorialSettings.cpp b/Engine/Source/Editor/IntroTutorials/Private/EditorTutorialSettings.cpp
--- a/Engine/Source/Editor/IntroTutorials/Private/EditorTutorialSettings.cpp
+++ b/Engine/Source/Editor/IntroTutorials/Private/EditorTutorialSettings.cpp
@@ -9,12 +9,14 @@ UEditorTutorialSettings::UEditorTutorialSettings(const class FPostConstructIniti
 
 }
 
-void UEditorTutorialSettings::FindTutorialsForContext(FName InContext, UEditorTutorial*& OutAttractTutorial, UEditorTutorial*& OutLaunchTutorial) const
+void UEditorTutorialSettings::FindTutorialInfoForContext(FName InContext, UEditorTutorial*& OutAttractTutorial, UEditorTutorial*& OutLaunchTutorial, FString& OutBrowserFilter) const
 {
 	for (const auto& Context : TutorialContexts)
 	{
 		if (Context.Context == InContext)
 		{
+			OutBrowserFilter = Context.BrowserFilter;
+
 			TSubclassOf<UEditorTutorial> LaunchTutorialClass = LoadClass<UEditorTutorial>(nullptr, *Context.LaunchTutorial.AssetLongPathname, nullptr, LOAD_None, nullptr);
 			if (LaunchTutorialClass != nullptr)
 			{
diff --git a/Engine/Source/Editor/IntroTutorials/Private/EditorTutorialSettings.h b/Engine/Source/Editor/IntroTutorials/Private/EditorTutorialSettings.h
--- a/Engine/Source/Editor/IntroTutorials/Private/EditorTutorialSettings.h
+++ b/Engine/Source/Editor/IntroTutorials/Private/EditorTutorialSettings.h
@@ -15,6 +15,10 @@ struct FTutorialContext
 	UPROPERTY(EditAnywhere, Category = "Tutorials")
 	FName Context;
 
+	/** The filter string to apply to the tutorials browser when launched from this context */
+	UPROPERTY(EditAnywhere, Category = "Tutorials")
+	FString BrowserFilter;
+
 	/** The tutorial to use in this context to let the user know there is a tutorial available */
 	UPROPERTY(EditAnywhere, Category = "Tutorials", meta = (MetaClass = "EditorTutorial"))
 	FStringClassReference AttractTutorial;
@@ -42,6 +46,6 @@ class UEditorTutorialSettings : public UObject
 	UPROPERTY(Config, EditAnywhere, Category = "Tutorials")
 	TArray<FTutorialContext> TutorialContexts;
 
-	/** Get the tutorial for the specified context */
-	void FindTutorialsForContext(FName InContext, UEditorTutorial*& OutAttractTutorial, UEditorTutorial*& OutLaunchTutorial) const;
+	/** Get the tutorial info for the specified context */
+	void FindTutorialInfoForContext(FName InContext, UEditorTutorial*& OutAttractTutorial, UEditorTutorial*& OutLaunchTutorial, FString& OutBrowserFilter) const;
 };
\ No newline at end of file
diff --git a/Engine/Source/Editor/IntroTutorials/Private/STutorialButton.cpp b/Engine/Source/Editor/IntroTutorials/Private/STutorialButton.cpp
--- a/Engine/Source/Editor/IntroTutorials/Private/STutorialButton.cpp
+++ b/Engine/Source/Editor/IntroTutorials/Private/STutorialButton.cpp
@@ -33,8 +33,7 @@ void STutorialButton::Construct(const FArguments& InArgs)
 		SNew(SButton)
 		.AddMetaData<FTutorialMetaData>(TagMeta)
 		.ButtonStyle(FEditorStyle::Get(), "TutorialLaunch.Button")
-		.ToolTipText(LOCTEXT("TutorialLaunchToolTip", "Launch Tutorial (right-click for more options)"))
-		.Visibility(this, &STutorialButton::GetVisibility)
+		.ToolTipText(this, &STutorialButton::GetButtonToolTip)
 		.OnClicked(this, &STutorialButton::HandleButtonClicked)
 		.ContentPadding(0.0f)
 		[
@@ -51,7 +50,8 @@ void STutorialButton::Tick(const FGeometry& AllottedGeometry, const double InCur
 	{
 		UEditorTutorial* AttractTutorial = nullptr;
 		UEditorTutorial* LaunchTutorial = nullptr;
-		GetDefault<UEditorTutorialSettings>()->FindTutorialsForContext(Context, AttractTutorial, LaunchTutorial);
+		FString BrowserFilter;
+		GetDefault<UEditorTutorialSettings>()->FindTutorialInfoForContext(Context, AttractTutorial, LaunchTutorial, BrowserFilter);
 
 		bTutorialAvailable = (LaunchTutorial != nullptr);
 		bTutorialCompleted = (LaunchTutorial != nullptr) && GetDefault<UTutorialStateSettings>()->HaveCompletedTutorial(LaunchTutorial);
@@ -83,7 +83,7 @@ int32 STutorialButton::OnPaint(const FPaintArgs& Args, const FGeometry& Allotted
 {
 	LayerId = SCompoundWidget::OnPaint(Args, AllottedGeometry, MyClippingRect, OutDrawElements, LayerId, InWidgetStyle, bParentEnabled) + 1000;
 
-	if (!(bTutorialCompleted || bTutorialDismissed))
+	if (ShouldShowAlert())
 	{
 		float AlphaFactor0 = 0.0f;
 		float AlphaFactor1 = 0.0f;
@@ -130,37 +130,59 @@ FReply STutorialButton::HandleButtonClicked()
 {
 	UEditorTutorial* AttractTutorial = nullptr;
 	UEditorTutorial* LaunchTutorial = nullptr;
-	GetDefault<UEditorTutorialSettings>()->FindTutorialsForContext(Context, AttractTutorial, LaunchTutorial);
-	if (LaunchTutorial != nullptr && ContextWindow.IsValid())
+	FString BrowserFilter;
+	GetDefault<UEditorTutorialSettings>()->FindTutorialInfoForContext(Context, AttractTutorial, LaunchTutorial, BrowserFilter);
+
+	bTutorialAvailable = (LaunchTutorial != nullptr);
+	bTutorialCompleted = (LaunchTutorial != nullptr) && GetDefault<UTutorialStateSettings>()->HaveCompletedTutorial(LaunchTutorial);
+	bTutorialDismissed = (AttractTutorial != nullptr) && GetDefault<UTutorialStateSettings>()->IsTutorialDismissed(AttractTutorial);
+
+	if(ContextWindow.IsValid())
 	{
 		FIntroTutorials& IntroTutorials = FModuleManager::GetModuleChecked<FIntroTutorials>(TEXT("IntroTutorials"));
-		const bool bRestart = true;
-		IntroTutorials.LaunchTutorial(LaunchTutorial, bRestart, ContextWindow);
+		if(ShouldLaunchBrowser())
+		{
+			IntroTutorials.SummonTutorialBrowser(ContextWindow.Pin().ToSharedRef(), BrowserFilter);
+		}
+		else if (LaunchTutorial != nullptr)
+		{
+			const bool bRestart = true;
+			IntroTutorials.LaunchTutorial(LaunchTutorial, bRestart, ContextWindow);
+		}
 	}
 
 	return FReply::Handled();
 }
 
-EVisibility STutorialButton::GetVisibility() const
-{
-	return bTutorialAvailable ? EVisibility::Visible : EVisibility::Collapsed;
-}
-
 FReply STutorialButton::OnMouseButtonDown(const FGeometry& MyGeometry, const FPointerEvent& MouseEvent)
 {
-	if (MouseEvent.GetEffectingButton() == EKeys::RightMouseButton)
+	if(!ShouldLaunchBrowser())
 	{
-		const bool bInShouldCloseWindowAfterMenuSelection = true;
-		FMenuBuilder MenuBuilder(bInShouldCloseWindowAfterMenuSelection, nullptr);
-
-		MenuBuilder.AddMenuEntry(
-			LOCTEXT("DismissReminder", "Dismiss Alert"),
-			LOCTEXT("DismissReminderTooltip", "Don't show me this alert again"),
-			FSlateIcon(),
-			FUIAction(FExecuteAction::CreateSP(this, &STutorialButton::DismissAlert))
-			);
-
-		FSlateApplication::Get().PushMenu(SharedThis(this), MenuBuilder.MakeWidget(), FSlateApplication::Get().GetCursorPos(), FPopupTransitionEffect::ContextMenu);
+		if (MouseEvent.GetEffectingButton() == EKeys::RightMouseButton)
+		{
+			const bool bInShouldCloseWindowAfterMenuSelection = true;
+			FMenuBuilder MenuBuilder(bInShouldCloseWindowAfterMenuSelection, nullptr);
+
+			if(ShouldShowAlert())
+			{
+				MenuBuilder.AddMenuEntry(
+					LOCTEXT("DismissReminder", "Dismiss Alert"),
+					LOCTEXT("DismissReminderTooltip", "Don't show me this alert again"),
+					FSlateIcon(),
+					FUIAction(FExecuteAction::CreateSP(this, &STutorialButton::DismissAlert))
+					);
+			}
+
+			MenuBuilder.AddMenuEntry(
+				LOCTEXT("LaunchBrowser", "Show Available Tutorials"),
+				LOCTEXT("LaunchBrowserTooltip", "Display the tutorials browser"),
+				FSlateIcon(),
+				FUIAction(FExecuteAction::CreateSP(this, &STutorialButton::LaunchBrowser))
+				);
+
+
+			FSlateApplication::Get().PushMenu(SharedThis(this), MenuBuilder.MakeWidget(), FSlateApplication::Get().GetCursorPos(), FPopupTransitionEffect::ContextMenu);
+		}
 	}
 	return FReply::Handled();
 }
@@ -169,7 +191,8 @@ void STutorialButton::DismissAlert()
 {
 	UEditorTutorial* AttractTutorial = nullptr;
 	UEditorTutorial* LaunchTutorial = nullptr;
-	GetDefault<UEditorTutorialSettings>()->FindTutorialsForContext(Context, AttractTutorial, LaunchTutorial);
+	FString BrowserFilter;
+	GetDefault<UEditorTutorialSettings>()->FindTutorialInfoForContext(Context, AttractTutorial, LaunchTutorial,BrowserFilter);
 	if (AttractTutorial != nullptr)
 	{
 		const bool bDismissAcrossSessions = true;
@@ -182,4 +205,38 @@ void STutorialButton::DismissAlert()
 	}
 }
 
+void STutorialButton::LaunchBrowser()
+{
+	if(ContextWindow.IsValid())
+	{
+		UEditorTutorial* AttractTutorial = nullptr;
+		UEditorTutorial* LaunchTutorial = nullptr;
+		FString BrowserFilter;
+		GetDefault<UEditorTutorialSettings>()->FindTutorialInfoForContext(Context, AttractTutorial, LaunchTutorial, BrowserFilter);
+
+		FIntroTutorials& IntroTutorials = FModuleManager::GetModuleChecked<FIntroTutorials>(TEXT("IntroTutorials"));
+		IntroTutorials.SummonTutorialBrowser(ContextWindow.Pin().ToSharedRef(), BrowserFilter);
+	}
+}
+
+bool STutorialButton::ShouldLaunchBrowser() const
+{
+	return (!bTutorialAvailable || (bTutorialAvailable && bTutorialCompleted));
+}
+
+bool STutorialButton::ShouldShowAlert() const
+{
+	return (bTutorialAvailable && !(bTutorialCompleted || bTutorialDismissed));
+}
+
+FText STutorialButton::GetButtonToolTip() const
+{
+	if(ShouldLaunchBrowser())
+	{
+		return LOCTEXT("TutorialLaunchBrowserToolTip", "Show Available Tutorials");
+	}
+
+	return LOCTEXT("TutorialLaunchToolTip", "Launch Tutorial (right-click for more options)");
+}
+
 #undef LOCTEXT_NAMESPACE
\ No newline at end of file
diff --git a/Engine/Source/Editor/IntroTutorials/Private/STutorialButton.h b/Engine/Source/Editor/IntroTutorials/Private/STutorialButton.h
--- a/Engine/Source/Editor/IntroTutorials/Private/STutorialButton.h
+++ b/Engine/Source/Editor/IntroTutorials/Private/STutorialButton.h
@@ -20,14 +20,24 @@ class STutorialButton : public SCompoundWidget
 	virtual void Tick(const FGeometry& AllottedGeometry, const double InCurrentTime, const float InDeltaTime) override;
 
 private:
-	void HandleMainFrameLoad(TSharedPtr<SWindow> InRootWindow, bool bIsNewProjectWindow);
-
+	/** Handle clicking the tutorial button */
 	FReply HandleButtonClicked();
 
-	EVisibility GetVisibility() const;
-
+	/** Dismiss the pulsing alert */
 	void DismissAlert();
 
+	/** Launch the tutorials browser */
+	void LaunchBrowser();
+
+	/** Check whether we should launch the browser in this context */
+	bool ShouldLaunchBrowser() const;
+
+	/** Check whether we should show the alert in this context */
+	bool ShouldShowAlert() const;
+
+	/** Get the tooltip for the tutorials button */
+	FText GetButtonToolTip() const;
+
 private:
 	/** Flag to defer tutorial open until the first Tick() */
 	bool bDeferTutorialOpen;
