Tidied up Tutorials Browser

Tweaked padding and button styles so the layout is more consistent.
Add ability to display textures as icons in the browser (for categories and tutorials).

diff --git a/Engine/Source/Editor/EditorStyle/Private/SlateEditorStyle.cpp b/Engine/Source/Editor/EditorStyle/Private/SlateEditorStyle.cpp
--- a/Engine/Source/Editor/EditorStyle/Private/SlateEditorStyle.cpp
+++ b/Engine/Source/Editor/EditorStyle/Private/SlateEditorStyle.cpp
@@ -2432,13 +2432,24 @@ void FSlateEditorStyle::FStyle::SetupTutorialStyles()
 		Set( "Tutorials.Browser.HighlightTextColor", TutorialSelectionColor );
 
 		Set( "Tutorials.Browser.Button", FButtonStyle()
+			.SetNormal( BOX_BRUSH( "Common/ButtonHoverHint", FMargin(4/16.0f), FLinearColor(1,1,1,1) ) )
+			.SetHovered( BOX_BRUSH( "Common/ButtonHoverHint", FMargin(4/16.0f), FLinearColor(0.9f,0.9f,0.9f,1) ) )
+			.SetPressed( BOX_BRUSH( "Common/ButtonHoverHint", FMargin(4/16.0f), FLinearColor(1,1,1,1) ) )
+			.SetDisabled( BOX_BRUSH( "Common/ButtonHoverHint", FMargin(4/16.0f), FLinearColor(1,1,1,1) ) )
+			.SetNormalPadding( FMargin(0,0,0,1))
+			.SetPressedPadding( FMargin(0,1,0,0)));
+
+		Set( "Tutorials.Content.Button", FButtonStyle()
 			.SetNormal( BOX_BRUSH( "Common/ButtonHoverHint", FMargin(4/16.0f), FLinearColor(1,1,1,0.15f) ) )
 			.SetHovered( BOX_BRUSH( "Common/ButtonHoverHint", FMargin(4/16.0f), FLinearColor(1,1,1,1) ) )
 			.SetPressed( BOX_BRUSH( "Common/ButtonHoverHint", FMargin(4/16.0f), FLinearColor(1,1,1,1) ) )
+			.SetDisabled( BOX_BRUSH( "Common/ButtonHoverHint", FMargin(4/16.0f), FLinearColor(1,1,1,1) ) )
 			.SetNormalPadding( FMargin(0,0,0,1))
 			.SetPressedPadding( FMargin(0,1,0,0)));
 
 		Set( "Tutorials.Browser.CategoryArrow", new IMAGE_BRUSH( "Tutorials/BrowserCategoryArrow", FVector2D(24.0f, 24.0f), FLinearColor(0.0f, 0.0f, 0.0f, 1.0f) ) );
+		Set( "Tutorials.Browser.DefaultTutorialIcon", new IMAGE_BRUSH( "Tutorials/DefaultTutorialIcon_40x", FVector2D(40.0f, 40.0f), FLinearColor::White ) );
+		Set( "Tutorials.Browser.DefaultCategoryIcon", new IMAGE_BRUSH( "Tutorials/DefaultCategoryIcon_40x", FVector2D(40.0f, 40.0f), FLinearColor::White ) );
 
 		Set( "Tutorials.Browser.BackButton", new IMAGE_BRUSH( "Tutorials/BrowserBack", FVector2D(32.0f, 32.0f), FLinearColor(0.0f, 0.0f, 0.0f, 1.0f) ) );
 		Set( "Tutorials.Browser.RestartButton", new IMAGE_BRUSH( "Tutorials/BrowserRestart", FVector2D(16.0f, 16.0f), FLinearColor(0.0f, 0.0f, 0.0f, 1.0f) ) );
@@ -2547,6 +2558,7 @@ void FSlateEditorStyle::FStyle::SetupTutorialStyles()
 			Set("TutorialEditableText.RoundedBackground", new BOX_BRUSH("Common/RoundedSelection_16x", 4.0f / 16.0f, FLinearColor(FColor(0xffeff3f3))));
 
 			Set("TutorialEditableText.Toolbar.HyperlinkImage", new IMAGE_BRUSH("Tutorials/hyperlink", Icon16x16, TextColor));
+			Set("TutorialEditableText.Toolbar.ImageImage", new IMAGE_BRUSH("Tutorials/Image", Icon16x16, TextColor));
 
 			Set("TutorialEditableText.Toolbar.TextColor", TextColor);
 
diff --git a/Engine/Source/Editor/IntroTutorials/Private/STutorialContent.cpp b/Engine/Source/Editor/IntroTutorials/Private/STutorialContent.cpp
--- a/Engine/Source/Editor/IntroTutorials/Private/STutorialContent.cpp
+++ b/Engine/Source/Editor/IntroTutorials/Private/STutorialContent.cpp
@@ -84,7 +84,7 @@ void STutorialContent::Construct(const FArguments& InArgs, const FTutorialConten
 						SNew(SButton)
 						.OnClicked(this, &STutorialContent::OnCloseButtonClicked)
 						.Visibility(this, &STutorialContent::GetCloseButtonVisibility)
-						.ButtonStyle(&FEditorStyle::Get().GetWidgetStyle<FButtonStyle>("Tutorials.Browser.Button"))
+						.ButtonStyle(&FEditorStyle::Get().GetWidgetStyle<FButtonStyle>("Tutorials.Content.Button"))
 						.ContentPadding(0.0f)
 						[
 							SNew(SImage)
diff --git a/Engine/Source/Editor/IntroTutorials/Private/STutorialNavigation.cpp b/Engine/Source/Editor/IntroTutorials/Private/STutorialNavigation.cpp
--- a/Engine/Source/Editor/IntroTutorials/Private/STutorialNavigation.cpp
+++ b/Engine/Source/Editor/IntroTutorials/Private/STutorialNavigation.cpp
@@ -31,7 +31,7 @@ void STutorialNavigation::Construct(const FArguments& InArgs)
 					SNew(SButton)
 					.OnClicked(this, &STutorialNavigation::OnBackButtonClicked)
 					.IsEnabled(InArgs._IsBackEnabled)
-					.ButtonStyle(&FEditorStyle::Get().GetWidgetStyle<FButtonStyle>("Tutorials.Browser.Button"))
+					.ButtonStyle(&FEditorStyle::Get().GetWidgetStyle<FButtonStyle>("Tutorials.Content.Button"))
 					.Content()
 					[
 						SNew(SImage)
@@ -46,7 +46,7 @@ void STutorialNavigation::Construct(const FArguments& InArgs)
 					SNew(SButton)
 					.OnClicked(this, &STutorialNavigation::OnHomeButtonClicked)
 					.IsEnabled(InArgs._IsHomeEnabled)
-					.ButtonStyle(&FEditorStyle::Get().GetWidgetStyle<FButtonStyle>("Tutorials.Browser.Button"))
+					.ButtonStyle(&FEditorStyle::Get().GetWidgetStyle<FButtonStyle>("Tutorials.Content.Button"))
 					.Content()
 					[
 						SNew(SImage)
@@ -61,7 +61,7 @@ void STutorialNavigation::Construct(const FArguments& InArgs)
 					SNew(SButton)
 					.OnClicked(this, &STutorialNavigation::OnNextButtonClicked)
 					.IsEnabled(InArgs._IsNextEnabled)
-					.ButtonStyle(&FEditorStyle::Get().GetWidgetStyle<FButtonStyle>("Tutorials.Browser.Button"))
+					.ButtonStyle(&FEditorStyle::Get().GetWidgetStyle<FButtonStyle>("Tutorials.Content.Button"))
 					.Content()
 					[
 						SNew(SImage)
diff --git a/Engine/Source/Editor/IntroTutorials/Private/STutorialsBrowser.cpp b/Engine/Source/Editor/IntroTutorials/Private/STutorialsBrowser.cpp
--- a/Engine/Source/Editor/IntroTutorials/Private/STutorialsBrowser.cpp
+++ b/Engine/Source/Editor/IntroTutorials/Private/STutorialsBrowser.cpp
@@ -19,13 +19,15 @@ class FTutorialListEntry_Category : public ITutorialListEntry, public TSharedFro
 {
 public:
 	FTutorialListEntry_Category(FOnCategorySelected InOnCategorySelected)
-		:  OnCategorySelected(InOnCategorySelected)
+		: OnCategorySelected(InOnCategorySelected)
+		, SlateBrush(nullptr)
 	{}
 
 	FTutorialListEntry_Category(const FTutorialCategory& InCategory, FOnCategorySelected InOnCategorySelected, const TAttribute<FText>& InHighlightText)
 		: Category(InCategory)
 		, OnCategorySelected(InOnCategorySelected)
 		, HighlightText(InHighlightText)
+		, SlateBrush(nullptr)
 	{
 		if(!Category.Identifier.IsEmpty())
 		{
@@ -39,6 +41,30 @@ public:
 				CategoryName = Category.Identifier;
 			}
 		}
+
+		if(Category.Texture.IsValid())
+		{
+			UTexture2D* Texture = LoadObject<UTexture2D>(nullptr, *Category.Texture.AssetLongPathname);
+			if(Texture != nullptr)
+			{
+				FIntPoint TextureSize = Texture->GetImportedSize();
+				DynamicBrush = MakeShareable(new FSlateDynamicImageBrush(Texture, FVector2D((float)TextureSize.X, (float)TextureSize.Y), NAME_None));
+				SlateBrush = DynamicBrush.Get();		
+			}
+		}	
+
+		if(SlateBrush == nullptr)
+		{
+			if(Category.Icon.Len() > 0)
+			{
+				SlateBrush = FEditorStyle::Get().GetBrush(FName(*Category.Icon));
+			}
+		}
+
+		if(SlateBrush == nullptr)
+		{
+			SlateBrush = FEditorStyle::Get().GetBrush("Tutorials.Browser.DefaultTutorialIcon");
+		}
 	}
 
 	virtual ~FTutorialListEntry_Category()
@@ -58,10 +84,17 @@ public:
 				.AutoWidth()
 				.VAlign(VAlign_Center)
 				.HAlign(HAlign_Center)
-				.Padding(16.0f)
+				.Padding(8.0f)
 				[
-					SNew(SImage)
-					.Image(FEditorStyle::Get().GetBrush(FName(*Category.Icon)))
+					SNew(SBox)
+					.WidthOverride(64.0f)
+					.HeightOverride(64.0f)
+					.VAlign(VAlign_Center)
+					.HAlign(HAlign_Center)
+					[
+						SNew(SImage)
+						.Image(SlateBrush)
+					]
 				]
 				+SHorizontalBox::Slot()
 				.FillWidth(1.0f)
@@ -160,6 +193,12 @@ public:
 
 	/** Text to highlight */
 	TAttribute<FText> HighlightText;
+
+	/** Static brush from the editor style */
+	const FSlateBrush* SlateBrush;
+
+	/** Dynamic brush from the texture specified by the user */
+	TSharedPtr<FSlateDynamicImageBrush> DynamicBrush;
 };
 
 DECLARE_DELEGATE_TwoParams(FOnTutorialSelected, UEditorTutorial* /* InTutorial */, bool /* bRestart */ );
@@ -171,7 +210,24 @@ public:
 		: Tutorial(InTutorial)
 		, OnTutorialSelected(InOnTutorialSelected)
 		, HighlightText(InHighlightText)
-	{}
+		, SlateBrush(nullptr)
+	{
+		if(Tutorial->Texture != nullptr)
+		{
+			FIntPoint TextureSize = Tutorial->Texture->GetImportedSize();
+			DynamicBrush = MakeShareable(new FSlateDynamicImageBrush(Tutorial->Texture, FVector2D((float)TextureSize.X, (float)TextureSize.Y), NAME_None));
+			SlateBrush = DynamicBrush.Get();
+		}	
+		else if(Tutorial->Icon.Len() > 0)
+		{
+			SlateBrush = FEditorStyle::Get().GetBrush(FName(*Tutorial->Icon));
+		}
+		
+		if(SlateBrush == nullptr)
+		{
+			SlateBrush = FEditorStyle::Get().GetBrush("Tutorials.Browser.DefaultTutorialIcon");
+		}
+	}
 
 	virtual ~FTutorialListEntry_Tutorial()
 	{}
@@ -194,15 +250,16 @@ public:
 					.AutoWidth()
 					.VAlign(VAlign_Center)
 					.HAlign(HAlign_Center)
+					.Padding(8.0f)
 					[
 						SNew(SBox)
-						.WidthOverride(96.0f)
-						.HeightOverride(96.0f)
+						.WidthOverride(64.0f)
+						.HeightOverride(64.0f)
 						.VAlign(VAlign_Center)
 						.HAlign(HAlign_Center)
 						[
 							SNew(SImage)
-							.Image(FEditorStyle::Get().GetBrush(FName(*Tutorial->Icon)))
+							.Image(SlateBrush)
 						]
 					]
 					+SHorizontalBox::Slot()
@@ -314,6 +371,12 @@ public:
 
 	/** Documentation page reference to use if we are displaying a UDN doc */
 	mutable TSharedPtr<IDocumentationPage> DocumentationPage;
+
+	/** Static brush from the editor style */
+	const FSlateBrush* SlateBrush;
+
+	/** Dynamic brush from the texture specified by the user */
+	TSharedPtr<FSlateDynamicImageBrush> DynamicBrush;
 };
 
 void STutorialsBrowser::Construct(const FArguments& InArgs)
@@ -346,7 +409,7 @@ void STutorialsBrowser::Construct(const FArguments& InArgs)
 						SNew(SButton)
 						.OnClicked(this, &STutorialsBrowser::OnBackButtonClicked)
 						.IsEnabled(this, &STutorialsBrowser::IsBackButtonEnabled)
-						.ButtonStyle(&FEditorStyle::Get().GetWidgetStyle<FButtonStyle>("Tutorials.Browser.Button"))
+						.ButtonStyle(&FEditorStyle::Get().GetWidgetStyle<FButtonStyle>("Tutorials.Content.Button"))
 						.Content()
 						[
 							SNew(SImage)
diff --git a/Engine/Source/Editor/IntroTutorials/Public/EditorTutorial.h b/Engine/Source/Editor/IntroTutorials/Public/EditorTutorial.h
--- a/Engine/Source/Editor/IntroTutorials/Public/EditorTutorial.h
+++ b/Engine/Source/Editor/IntroTutorials/Public/EditorTutorial.h
@@ -59,9 +59,13 @@ struct INTROTUTORIALS_API FTutorialCategory
 	UPROPERTY(EditAnywhere, Category="Content", meta=(MultiLine=true))
 	FText Description;
 
-	/** Icon for this tutorial, used when presented to the user */
+	/** Icon for this tutorial, used when presented to the user in the tutorial browser. Only used if there isn't a valid texture to use. */
 	UPROPERTY(EditAnywhere, Category="Tutorial")
 	FString Icon;
+
+	/** Texture for this tutorial, used when presented to the user in the tutorial browser. */
+	UPROPERTY(EditAnywhere, Category="Tutorial", meta=(AllowedClasses="Texture2D"))
+	FStringAssetReference Texture;
 };
 
 /** Content wrapper */
@@ -204,10 +208,14 @@ class INTROTUTORIALS_API UEditorTutorial : public UObject
 	UPROPERTY(EditAnywhere, Category="Tutorial")
 	FText Title;
 
-	/** Icon for this tutorial, used when presented to the user */
+	/** Icon name for this tutorial, used when presented to the user in the tutorial browser. This is a name for the icon in the Slate editor style. Only used if there isn't a valid texture to use. */
 	UPROPERTY(EditAnywhere, Category="Tutorial")
 	FString Icon;
 
+	/** Texture for this tutorial, used when presented to the user in the tutorial browser. */
+	UPROPERTY(EditAnywhere, Category="Tutorial")
+	UTexture2D* Texture;
+
 	/** Category of this tutorial, used to organize tutorials when presented to the user */
 	UPROPERTY(EditAnywhere, Category="Tutorial")
 	FString Category;
