Fixed crash when zero-length tutorials are run

TTP# 345094 - CRASH: Selecting the arrow keys in the tutorial will cause an immediate crash with the new Tutorial system

diff --git a/Engine/Source/Editor/IntroTutorials/Private/STutorialRoot.cpp b/Engine/Source/Editor/IntroTutorials/Private/STutorialRoot.cpp
--- a/Engine/Source/Editor/IntroTutorials/Private/STutorialRoot.cpp
+++ b/Engine/Source/Editor/IntroTutorials/Private/STutorialRoot.cpp
@@ -209,7 +209,10 @@ void STutorialRoot::GoToNextStage(TWeakPtr<SWindow> InNavigationWindow)
 	{
 		UEditorTutorial* PreviousTutorial = CurrentTutorial;
 		int32 PreviousTutorialStage = CurrentTutorialStage;
-		CurrentTutorial->HandleTutorialStageEnded(CurrentTutorial->Stages[CurrentTutorialStage].Name);
+		if(CurrentTutorialStage < CurrentTutorial->Stages.Num())
+		{
+			CurrentTutorial->HandleTutorialStageEnded(CurrentTutorial->Stages[CurrentTutorialStage].Name);
+		}
 
 		if(CurrentTutorialStage + 1 >= CurrentTutorial->Stages.Num() && FName(*CurrentTutorial->NextTutorial.AssetLongPathname) != NAME_None)
 		{
