Rich tooltips for tutorial hyperlinks

Refactored handling of decorators into new utility class FTutorialText.
Added tooltip (and rich tooltip) support to hyperlink runs.

diff --git a/Engine/Source/Editor/IntroTutorials/Private/STutorialContent.cpp b/Engine/Source/Editor/IntroTutorials/Private/STutorialContent.cpp
--- a/Engine/Source/Editor/IntroTutorials/Private/STutorialContent.cpp
+++ b/Engine/Source/Editor/IntroTutorials/Private/STutorialContent.cpp
@@ -4,6 +4,7 @@
 #include "STutorialContent.h"
 #include "STutorialEditableText.h"
 #include "IntroTutorials.h"
+#include "TutorialText.h"
 
 #define LOCTEXT_NAMESPACE "STutorialContent"
 
@@ -239,12 +240,7 @@ TSharedRef<SWidget> STutorialContent::GenerateContentWidget(const FTutorialConte
 	case ETutorialContent::RichText:
 		{
 			TArray< TSharedRef< class ITextDecorator > > Decorators;
-			Decorators.Add(FHyperlinkDecorator::Create(TEXT("browser"), FSlateHyperlinkRun::FOnClick::CreateStatic(&TutorialTextHelpers::OnBrowserLinkClicked)));
-			Decorators.Add(FHyperlinkDecorator::Create(TEXT("udn"), FSlateHyperlinkRun::FOnClick::CreateStatic(&TutorialTextHelpers::OnDocLinkClicked)));
-			Decorators.Add(FHyperlinkDecorator::Create(TEXT("tutorial"), FSlateHyperlinkRun::FOnClick::CreateStatic(&TutorialTextHelpers::OnTutorialLinkClicked)));
-			Decorators.Add(FHyperlinkDecorator::Create(TEXT("code"), FSlateHyperlinkRun::FOnClick::CreateStatic(&TutorialTextHelpers::OnCodeLinkClicked)));
-			Decorators.Add(FHyperlinkDecorator::Create(TEXT("asset"), FSlateHyperlinkRun::FOnClick::CreateStatic(&TutorialTextHelpers::OnAssetLinkClicked)));
-			Decorators.Add(FTextStyleDecorator::Create());
+			FTutorialText::GetRichTextDecorators(Decorators);
 
 			return SNew(SRichTextBlock)
 					.Visibility(EVisibility::SelfHitTestInvisible)
diff --git a/Engine/Source/Editor/IntroTutorials/Private/STutorialEditableText.cpp b/Engine/Source/Editor/IntroTutorials/Private/STutorialEditableText.cpp
--- a/Engine/Source/Editor/IntroTutorials/Private/STutorialEditableText.cpp
+++ b/Engine/Source/Editor/IntroTutorials/Private/STutorialEditableText.cpp
@@ -2,154 +2,11 @@
 
 #include "IntroTutorialsPrivatePCH.h"
 #include "STutorialEditableText.h"
-#include "IDocumentation.h"
-#include "ISourceCodeAccessModule.h"
-#include "ContentBrowserModule.h"
-#include "SColorPicker.h"
-#include "DesktopPlatformModule.h"
 #include "RichTextLayoutMarshaller.h"
-#include "EditorTutorial.h"
-#include "Editor/MainFrame/Public/MainFrame.h"
+#include "TutorialText.h"
 
 #define LOCTEXT_NAMESPACE "STutorialEditableText"
 
-
-namespace TutorialTextHelpers
-{
-
-void OnBrowserLinkClicked(const FSlateHyperlinkRun::FMetadata& Metadata)
-{
-	const FString* Url = Metadata.Find(TEXT("href"));
-	if(Url)
-	{
-		FPlatformProcess::LaunchURL(**Url, nullptr, nullptr);
-	}
-}
-
-
-void OnDocLinkClicked(const FSlateHyperlinkRun::FMetadata& Metadata)
-{
-	const FString* Url = Metadata.Find(TEXT("href"));
-	if(Url)
-	{
-		IDocumentation::Get()->Open(*Url);
-	}
-}
-
-static void ParseTutorialLink(const FString &InternalLink)
-{
-	UBlueprint* Blueprint = LoadObject<UBlueprint>(nullptr, *InternalLink);
-	if (Blueprint && Blueprint->GeneratedClass)
-	{
-		FIntroTutorials& IntroTutorials = FModuleManager::GetModuleChecked<FIntroTutorials>(TEXT("IntroTutorials"));
-		IMainFrameModule& MainFrameModule = FModuleManager::LoadModuleChecked<IMainFrameModule>(TEXT("MainFrame"));
-		const bool bRestart = true;
-		IntroTutorials.LaunchTutorial(Blueprint->GeneratedClass->GetDefaultObject<UEditorTutorial>(), bRestart, MainFrameModule.GetParentWindow());
-	}
-}
-
-void OnTutorialLinkClicked(const FSlateHyperlinkRun::FMetadata& Metadata)
-{
-	const FString* Url = Metadata.Find(TEXT("href"));
-	if(Url)
-	{
-		ParseTutorialLink(*Url);
-	}
-}
-
-static void ParseCodeLink(const FString &InternalLink)
-{
-	// Tokens used by the code parsing. Details in the parse section	
-	static const FString ProjectSpecifier(TEXT("[PROJECTPATH]"));
-	static const FString ProjectPathSpecifier(TEXT("[PROJECT]"));
-	static const FString EnginePathSpecifier(TEXT("[ENGINEPATH]"));
-
-	FString Path;
-	int32 Line = 0;
-	int32 Col = 0;
-
-	TArray<FString> Tokens;
-	InternalLink.ParseIntoArray(&Tokens, TEXT(","), 0);
-	int32 TokenStringsCount = Tokens.Num();
-	if (TokenStringsCount > 0)
-	{
-		Path = Tokens[0];
-	}
-	if (TokenStringsCount > 1)
-	{
-		TTypeFromString<int32>::FromString(Line, *Tokens[1]);
-	}
-	if (TokenStringsCount > 2)
-	{
-		TTypeFromString<int32>::FromString(Col, *Tokens[2]);
-	}
-
-	ISourceCodeAccessModule& SourceCodeAccessModule = FModuleManager::LoadModuleChecked<ISourceCodeAccessModule>("SourceCodeAccess");
-	ISourceCodeAccessor& SourceCodeAccessor = SourceCodeAccessModule.GetAccessor();
-
-	if (Path.Contains(EnginePathSpecifier) == true)
-	{
-		// replace engine path specifier with path to engine
-		Path.ReplaceInline(*EnginePathSpecifier, *FPaths::EngineDir());
-	}
-
-	if (Path.Contains(ProjectSpecifier) == true)
-	{
-		// replace project specifier with path to project
-		Path.ReplaceInline(*ProjectSpecifier, FApp::GetGameName());
-	}
-
-	if (Path.Contains(ProjectPathSpecifier) == true)
-	{
-		// replace project specifier with path to project
-		Path.ReplaceInline(*ProjectPathSpecifier, *FPaths::GetProjectFilePath());
-	}
-
-	Path = FPaths::ConvertRelativePathToFull(Path);
-
-	SourceCodeAccessor.OpenFileAtLine(Path, Line, Col);
-}
-
-void OnCodeLinkClicked(const FSlateHyperlinkRun::FMetadata& Metadata)
-{
-	const FString* Url = Metadata.Find(TEXT("href"));
-	if(Url)
-	{	
-		ParseCodeLink(*Url);
-	}
-}
-
-static void ParseAssetLink(const FString& InternalLink, const FString* Action)
-{
-	UObject* RequiredObject = FindObject<UObject>(ANY_PACKAGE, *InternalLink);
-	if (RequiredObject != nullptr)
-	{
-		if(Action && *Action == TEXT("select"))
-		{
-			FContentBrowserModule& ContentBrowserModule = FModuleManager::Get().LoadModuleChecked<FContentBrowserModule>("ContentBrowser");
-			TArray<UObject*> AssetToBrowse;
-			AssetToBrowse.Add(RequiredObject);
-			ContentBrowserModule.Get().SyncBrowserToAssets(AssetToBrowse);
-		}
-		else
-		{
-			FAssetEditorManager::Get().OpenEditorForAsset(RequiredObject);
-		}
-	}
-}
-
-void OnAssetLinkClicked(const FSlateHyperlinkRun::FMetadata& Metadata)
-{
-	const FString* Url = Metadata.Find(TEXT("href"));
-	const FString* Action = Metadata.Find(TEXT("action"));
-	if(Url)
-	{
-		ParseAssetLink(*Url, Action);
-	}
-}
-
-}
-
 BEGIN_SLATE_FUNCTION_BUILD_OPTIMIZATION
 void STutorialEditableText::Construct(const FArguments& InArgs)
 {
@@ -174,49 +31,16 @@ void STutorialEditableText::Construct(const FArguments& InArgs)
 		FEditorStyle::Get().GetWidgetStyle<FTextBlockStyle>("Tutorials.Content.Text")
 		);
 
-	HyperlinkDescs.Add(MakeShareable(new FHyperlinkTypeDesc(
-		EHyperlinkType::Browser, 
-		LOCTEXT("BrowserLinkTypeLabel", "URL"), 
-		LOCTEXT("BrowserLinkTypeTooltip", "A link that opens a browser window (e.g. http://www.unrealengine.com)"),
-		TEXT("browser"),
-		FSlateHyperlinkRun::FOnClick::CreateStatic(&TutorialTextHelpers::OnBrowserLinkClicked))));
-
-	CurrentHyperlinkType = HyperlinkDescs[0];
-
-	HyperlinkDescs.Add(MakeShareable(new FHyperlinkTypeDesc(
-		EHyperlinkType::UDN, 
-		LOCTEXT("UDNLinkTypeLabel", "UDN"), 
-		LOCTEXT("UDNLinkTypeTooltip", "A link that opens some UDN documentation (e.g. /Engine/Blueprints/UserGuide/Types/ClassBlueprint)"),
-		TEXT("udn"),
-		FSlateHyperlinkRun::FOnClick::CreateStatic(&TutorialTextHelpers::OnDocLinkClicked))));
-
-	HyperlinkDescs.Add(MakeShareable(new FHyperlinkTypeDesc(
-		EHyperlinkType::Asset, 
-		LOCTEXT("AssetLinkTypeLabel", "Asset"), 
-		LOCTEXT("AssetLinkTypeTooltip", "A link that opens an asset (e.g. /Game/StaticMeshes/SphereMesh.SphereMesh)"),
-		TEXT("asset"),
-		FSlateHyperlinkRun::FOnClick::CreateStatic(&TutorialTextHelpers::OnAssetLinkClicked))));
-
-	HyperlinkDescs.Add(MakeShareable(new FHyperlinkTypeDesc(
-		EHyperlinkType::Code, 
-		LOCTEXT("CodeLinkTypeLabel", "Code"), 
-		LOCTEXT("CodeLinkTypeTooltip", "A link that opens code in your selected IDE.\nFor example: [PROJECTPATH]/Private/SourceFile.cpp,1,1.\nThe numbers correspond to line number and column number.\nYou can use [PROJECT], [PROJECTPATH] and [ENGINEPATH] tags to make paths."),
-		TEXT("code"),
-		FSlateHyperlinkRun::FOnClick::CreateStatic(&TutorialTextHelpers::OnCodeLinkClicked))));
-
-	HyperlinkDescs.Add(MakeShareable(new FHyperlinkTypeDesc(
-		EHyperlinkType::Tutorial, 
-		LOCTEXT("TutorialLinkTypeLabel", "Tutorial"), 
-		LOCTEXT("TutorialLinkTypeTooltip", "A type of asset link that opens another tutorial, e.g. /Game/Tutorials/StaticMeshTutorial.StaticMeshTutorial"),
-		TEXT("tutorial"),
-		FSlateHyperlinkRun::FOnClick::CreateStatic(&TutorialTextHelpers::OnTutorialLinkClicked))));
-
-	for(const auto& HyperlinkDesc : HyperlinkDescs)
+
+	TArray<TSharedRef<class ITextDecorator>> Decorators;
+	FTutorialText::GetRichTextDecorators(Decorators);
+
+	for(auto& Decorator : Decorators)
 	{
-		RichTextMarshaller->AppendInlineDecorator(FHyperlinkDecorator::Create(HyperlinkDesc->Id, HyperlinkDesc->OnClickedDelegate));
+		RichTextMarshaller->AppendInlineDecorator(Decorator);
 	}
 
-	RichTextMarshaller->AppendInlineDecorator(FTextStyleDecorator::Create());
+	CurrentHyperlinkType = FTutorialText::GetHyperLinkDescs()[0];
 
 	this->ChildSlot
 	[
@@ -350,12 +174,12 @@ void STutorialEditableText::Construct(const FArguments& InArgs)
 							.VAlign(VAlign_Center)
 							[
 								SNew(SComboBox<TSharedPtr<FHyperlinkTypeDesc>>)
-								.OptionsSource(&HyperlinkDescs)
+								.OptionsSource(&FTutorialText::GetHyperLinkDescs())
 								.ComboBoxStyle(FEditorStyle::Get(), "TutorialEditableText.Toolbar.ComboBox")
 								.OnSelectionChanged(this, &STutorialEditableText::OnActiveHyperlinkChanged)
 								.OnGenerateWidget(this, &STutorialEditableText::GenerateHyperlinkComboEntry)
 								.ContentPadding(0.0f)
-								.InitiallySelectedItem(HyperlinkDescs[0])
+								.InitiallySelectedItem(FTutorialText::GetHyperLinkDescs()[0])
 								.Content()
 								[
 									SNew(SBox)
@@ -386,6 +210,16 @@ void STutorialEditableText::Construct(const FArguments& InArgs)
 								]
 							]
 							+SHorizontalBox::Slot()
+							.AutoWidth()
+							.VAlign(VAlign_Center)
+							.Padding(5.0f, 0.0f, 0.0f, 0.0f)
+							[
+								SAssignNew(UDNExcerptTextBox, SEditableTextBox)
+								.HintText(LOCTEXT("ExcerptHint", "Excerpt"))
+								.ToolTipText(LOCTEXT("ExcerptAssetTooltip", "Enter the excerpt that should be used for this link's rich tooltip"))
+								.Visibility(this, &STutorialEditableText::GetExcerptVisibility)
+							]
+							+SHorizontalBox::Slot()
 							.FillWidth(1.0f)
 							.VAlign(VAlign_Center)
 							.HAlign(HAlign_Right)
@@ -551,7 +385,7 @@ void STutorialEditableText::StyleSelectedText()
 TSharedPtr<FHyperlinkTypeDesc> STutorialEditableText::GetHyperlinkTypeFromId(const FString& InId) const
 {
 	TSharedPtr<FHyperlinkTypeDesc> FoundType;
-	for(const auto& Desc : HyperlinkDescs)
+	for(const auto& Desc : FTutorialText::GetHyperLinkDescs())
 	{
 		if(Desc->Id == InId)
 		{
@@ -589,7 +423,7 @@ void STutorialEditableText::HandleHyperlinkComboOpened()
 		HyperlinkURLTextBox->SetText((URLUnderCursor) ? FText::FromString(*URLUnderCursor) : FText());
 
 		const FString* const IdUnderCursor = Run->GetRunInfo().MetaData.Find(TEXT("id"));
-		CurrentHyperlinkType = IdUnderCursor ? GetHyperlinkTypeFromId(*IdUnderCursor) : HyperlinkDescs[0];
+		CurrentHyperlinkType = IdUnderCursor ? GetHyperlinkTypeFromId(*IdUnderCursor) : FTutorialText::GetHyperLinkDescs()[0];
 
 		FString RunText;
 		Run->AppendTextTo(RunText);
@@ -621,6 +455,10 @@ FReply STutorialEditableText::HandleInsertHyperLinkClicked()
 		{
 			RunInfo.MetaData.Add(TEXT("action"), bOpenAsset ? TEXT("edit") : TEXT("select"));
 		}
+		else if(CurrentHyperlinkType->Type == EHyperlinkType::UDN && !UDNExcerptTextBox->GetText().IsEmpty())
+		{
+			RunInfo.MetaData.Add(TEXT("excerpt"), UDNExcerptTextBox->GetText().ToString());
+		}
 
 		// Create the new run, and then insert it at the cursor position
 		TSharedRef<FSlateHyperlinkRun> HyperlinkRun = FSlateHyperlinkRun::Create(
@@ -702,4 +540,9 @@ ESlateCheckBoxState::Type STutorialEditableText::IsOpenAssetChecked() const
 	return bOpenAsset ? ESlateCheckBoxState::Checked : ESlateCheckBoxState::Unchecked;
 }
 
+EVisibility STutorialEditableText::GetExcerptVisibility() const
+{
+	return CurrentHyperlinkType.IsValid() && CurrentHyperlinkType->Type == EHyperlinkType::UDN ? EVisibility::Visible : EVisibility::Collapsed;
+}
+
 #undef LOCTEXT_NAMESPACE
\ No newline at end of file
diff --git a/Engine/Source/Editor/IntroTutorials/Private/STutorialEditableText.h b/Engine/Source/Editor/IntroTutorials/Private/STutorialEditableText.h
--- a/Engine/Source/Editor/IntroTutorials/Private/STutorialEditableText.h
+++ b/Engine/Source/Editor/IntroTutorials/Private/STutorialEditableText.h
@@ -2,31 +2,6 @@
 
 #pragma once
 
-namespace TutorialTextHelpers
-{
-	extern void OnBrowserLinkClicked(const FSlateHyperlinkRun::FMetadata& Metadata);
-
-	extern void OnDocLinkClicked(const FSlateHyperlinkRun::FMetadata& Metadata);
-
-	extern void OnTutorialLinkClicked(const FSlateHyperlinkRun::FMetadata& Metadata);
-
-	extern void OnCodeLinkClicked(const FSlateHyperlinkRun::FMetadata& Metadata);
-
-	extern void OnAssetLinkClicked(const FSlateHyperlinkRun::FMetadata& Metadata);
-}
-
-namespace EHyperlinkType
-{
-	enum Type
-	{
-		Browser,
-		UDN,
-		Tutorial,
-		Code,
-		Asset
-	};
-}
-
 /** Text style and name to display in the UI */
 struct FTextStyleAndName
 {
@@ -138,73 +113,6 @@ struct FTextStyleAndName
 	FText DisplayName;
 };
 
-/**
- * This is a custom decorator used to allow arbitrary styling of text within a rich-text editor
- * This is required since normal text styling can only work with known styles from a given Slate style-set
- */
-class FTextStyleDecorator : public ITextDecorator
-{
-public:
-
-	static TSharedRef<FTextStyleDecorator> Create()
-	{
-		return MakeShareable(new FTextStyleDecorator());
-	}
-
-	virtual ~FTextStyleDecorator()
-	{
-	}
-
-	virtual bool Supports(const FTextRunParseResults& RunParseResult, const FString& Text) const override
-	{
-		return (RunParseResult.Name == TEXT("TextStyle"));
-	}
-
-	virtual TSharedRef<ISlateRun> Create(const TSharedRef<class FTextLayout>& TextLayout, const FTextRunParseResults& RunParseResult, const FString& OriginalText, const TSharedRef< FString >& InOutModelText, const ISlateStyle* Style) override
-	{
-		FRunInfo RunInfo(RunParseResult.Name);
-		for(const TPair<FString, FTextRange>& Pair : RunParseResult.MetaData)
-		{
-			RunInfo.MetaData.Add(Pair.Key, OriginalText.Mid(Pair.Value.BeginIndex, Pair.Value.EndIndex - Pair.Value.BeginIndex));
-		}
-
-		FTextRange ModelRange;
-		ModelRange.BeginIndex = InOutModelText->Len();
-		*InOutModelText += OriginalText.Mid(RunParseResult.ContentRange.BeginIndex, RunParseResult.ContentRange.EndIndex - RunParseResult.ContentRange.BeginIndex);
-		ModelRange.EndIndex = InOutModelText->Len();
-
-		return FSlateTextRun::Create(RunInfo, InOutModelText, FTextStyleAndName::CreateTextBlockStyle(RunInfo), ModelRange);
-	}
-};
-
-/** Helper struct to hold info about hyperlink types */
-struct FHyperlinkTypeDesc
-{
-	FHyperlinkTypeDesc(EHyperlinkType::Type InType, const FText& InText, const FText& InTooltipText, const FString& InId, FSlateHyperlinkRun::FOnClick InOnClickedDelegate)
-		: Type(InType)
-		, Id(InId)
-		, Text(InText)
-		, TooltipText(InTooltipText)
-		, OnClickedDelegate(InOnClickedDelegate)
-	{
-	}
-
-	/** The type of the link */
-	EHyperlinkType::Type Type;
-
-	/** Tag used by this hyperlink's run */
-	FString Id;
-
-	/** Text to display in the UI */
-	FText Text;
-
-	/** Tooltip text to display in the UI */
-	FText TooltipText;
-
-	/** Delegate to execute for this hyperlink's run */
-	FSlateHyperlinkRun::FOnClick OnClickedDelegate;
-};
-
 
 class STutorialEditableText : public SCompoundWidget
 {
@@ -256,9 +164,9 @@ protected:
 
 	FText GetHyperlinkButtonText() const;
 
-	void OnActiveHyperlinkChanged(TSharedPtr<FHyperlinkTypeDesc> NewValue, ESelectInfo::Type SelectionType);
+	void OnActiveHyperlinkChanged(TSharedPtr<struct FHyperlinkTypeDesc> NewValue, ESelectInfo::Type SelectionType);
 
-	TSharedRef<SWidget> GenerateHyperlinkComboEntry(TSharedPtr<FHyperlinkTypeDesc> SourceEntry);
+	TSharedRef<SWidget> GenerateHyperlinkComboEntry(TSharedPtr<struct FHyperlinkTypeDesc> SourceEntry);
 
 	FText GetActiveHyperlinkName() const;
 
@@ -272,6 +180,8 @@ protected:
 
 	ESlateCheckBoxState::Type IsOpenAssetChecked() const;
 
+	EVisibility GetExcerptVisibility() const;
+
 protected:
 	TSharedPtr<SMultiLineEditableTextBox> RichEditableTextBox;
 
@@ -285,6 +195,7 @@ protected:
 	TSharedPtr<SComboBox<TSharedPtr<FTextStyleAndName>>> FontComboBox;
 	TSharedPtr<STextBlock> HyperlinkNameTextBlock;
 	TSharedPtr<SEditableTextBox> HyperlinkURLTextBox;
+	TSharedPtr<SEditableTextBox> UDNExcerptTextBox;
 
 	TSharedPtr<FTextStyleAndName> ActiveStyle;
 	TSharedPtr<FTextStyleAndName> HyperlinkStyle;
@@ -294,8 +205,7 @@ protected:
 	FOnTextCommitted OnTextCommitted;
 	FOnTextChanged OnTextChanged;
 
-	TArray<TSharedPtr<FHyperlinkTypeDesc>> HyperlinkDescs;
-	TSharedPtr<FHyperlinkTypeDesc> CurrentHyperlinkType;
+	TSharedPtr<struct FHyperlinkTypeDesc> CurrentHyperlinkType;
 
 	bool bOpenAsset;
 
diff --git a/Engine/Source/Editor/IntroTutorials/Private/TutorialText.cpp b/Engine/Source/Editor/IntroTutorials/Private/TutorialText.cpp
new file mode 100644
--- /dev/null
+++ b/Engine/Source/Editor/IntroTutorials/Private/TutorialText.cpp
@@ -0,0 +1,323 @@
+// Copyright 1998-2014 Epic Games, Inc. All Rights Reserved.
+
+#include "IntroTutorialsPrivatePCH.h"
+#include "TutorialText.h"
+#include "IDocumentation.h"
+#include "ISourceCodeAccessModule.h"
+#include "ContentBrowserModule.h"
+#include "DesktopPlatformModule.h"
+#include "Editor/MainFrame/Public/MainFrame.h"
+#include "EditorTutorial.h"
+#include "SourceCodeNavigation.h"
+
+#define LOCTEXT_NAMESPACE "TutorialText"
+
+TArray<TSharedPtr<FHyperlinkTypeDesc>> FTutorialText::HyperlinkDescs;
+
+
+/**
+ * This is a custom decorator used to allow arbitrary styling of text within a rich-text editor
+ * This is required since normal text styling can only work with known styles from a given Slate style-set
+ */
+class FTextStyleDecorator : public ITextDecorator
+{
+public:
+
+	static TSharedRef<FTextStyleDecorator> Create()
+	{
+		return MakeShareable(new FTextStyleDecorator());
+	}
+
+	virtual ~FTextStyleDecorator()
+	{
+	}
+
+	virtual bool Supports(const FTextRunParseResults& RunParseResult, const FString& Text) const override
+	{
+		return (RunParseResult.Name == TEXT("TextStyle"));
+	}
+
+	virtual TSharedRef<ISlateRun> Create(const TSharedRef<class FTextLayout>& TextLayout, const FTextRunParseResults& RunParseResult, const FString& OriginalText, const TSharedRef< FString >& InOutModelText, const ISlateStyle* Style) override
+	{
+		FRunInfo RunInfo(RunParseResult.Name);
+		for(const TPair<FString, FTextRange>& Pair : RunParseResult.MetaData)
+		{
+			RunInfo.MetaData.Add(Pair.Key, OriginalText.Mid(Pair.Value.BeginIndex, Pair.Value.EndIndex - Pair.Value.BeginIndex));
+		}
+
+		FTextRange ModelRange;
+		ModelRange.BeginIndex = InOutModelText->Len();
+		*InOutModelText += OriginalText.Mid(RunParseResult.ContentRange.BeginIndex, RunParseResult.ContentRange.EndIndex - RunParseResult.ContentRange.BeginIndex);
+		ModelRange.EndIndex = InOutModelText->Len();
+
+		return FSlateTextRun::Create(RunInfo, InOutModelText, FTextStyleAndName::CreateTextBlockStyle(RunInfo), ModelRange);
+	}
+};
+
+
+static void OnBrowserLinkClicked(const FSlateHyperlinkRun::FMetadata& Metadata)
+{
+	const FString* Url = Metadata.Find(TEXT("href"));
+	if(Url)
+	{
+		FPlatformProcess::LaunchURL(**Url, nullptr, nullptr);
+	}
+}
+
+static void OnDocLinkClicked(const FSlateHyperlinkRun::FMetadata& Metadata)
+{
+	const FString* Url = Metadata.Find(TEXT("href"));
+	if(Url)
+	{
+		IDocumentation::Get()->Open(*Url);
+	}
+}
+
+static void ParseTutorialLink(const FString &InternalLink)
+{
+	UBlueprint* Blueprint = LoadObject<UBlueprint>(nullptr, *InternalLink);
+	if (Blueprint && Blueprint->GeneratedClass)
+	{
+		FIntroTutorials& IntroTutorials = FModuleManager::GetModuleChecked<FIntroTutorials>(TEXT("IntroTutorials"));
+		IMainFrameModule& MainFrameModule = FModuleManager::LoadModuleChecked<IMainFrameModule>(TEXT("MainFrame"));
+		const bool bRestart = true;
+		IntroTutorials.LaunchTutorial(Blueprint->GeneratedClass->GetDefaultObject<UEditorTutorial>(), bRestart, MainFrameModule.GetParentWindow());
+	}
+}
+
+static void OnTutorialLinkClicked(const FSlateHyperlinkRun::FMetadata& Metadata)
+{
+	const FString* Url = Metadata.Find(TEXT("href"));
+	if(Url)
+	{
+		ParseTutorialLink(*Url);
+	}
+}
+
+static void ParseCodeLink(const FString &InternalLink)
+{
+	// Tokens used by the code parsing. Details in the parse section	
+	static const FString ProjectSpecifier(TEXT("[PROJECTPATH]"));
+	static const FString ProjectPathSpecifier(TEXT("[PROJECT]"));
+	static const FString EnginePathSpecifier(TEXT("[ENGINEPATH]"));
+
+	FString Path;
+	int32 Line = 0;
+	int32 Col = 0;
+
+	TArray<FString> Tokens;
+	InternalLink.ParseIntoArray(&Tokens, TEXT(","), 0);
+	int32 TokenStringsCount = Tokens.Num();
+	if (TokenStringsCount > 0)
+	{
+		Path = Tokens[0];
+	}
+	if (TokenStringsCount > 1)
+	{
+		TTypeFromString<int32>::FromString(Line, *Tokens[1]);
+	}
+	if (TokenStringsCount > 2)
+	{
+		TTypeFromString<int32>::FromString(Col, *Tokens[2]);
+	}
+
+	ISourceCodeAccessModule& SourceCodeAccessModule = FModuleManager::LoadModuleChecked<ISourceCodeAccessModule>("SourceCodeAccess");
+	ISourceCodeAccessor& SourceCodeAccessor = SourceCodeAccessModule.GetAccessor();
+
+	if (Path.Contains(EnginePathSpecifier) == true)
+	{
+		// replace engine path specifier with path to engine
+		Path.ReplaceInline(*EnginePathSpecifier, *FPaths::EngineDir());
+	}
+
+	if (Path.Contains(ProjectSpecifier) == true)
+	{
+		// replace project specifier with path to project
+		Path.ReplaceInline(*ProjectSpecifier, FApp::GetGameName());
+	}
+
+	if (Path.Contains(ProjectPathSpecifier) == true)
+	{
+		// replace project specifier with path to project
+		Path.ReplaceInline(*ProjectPathSpecifier, *FPaths::GetProjectFilePath());
+	}
+
+	Path = FPaths::ConvertRelativePathToFull(Path);
+
+	SourceCodeAccessor.OpenFileAtLine(Path, Line, Col);
+}
+
+static void OnCodeLinkClicked(const FSlateHyperlinkRun::FMetadata& Metadata)
+{
+	const FString* Url = Metadata.Find(TEXT("href"));
+	if(Url)
+	{	
+		ParseCodeLink(*Url);
+	}
+}
+
+static void ParseAssetLink(const FString& InternalLink, const FString* Action)
+{
+	UObject* RequiredObject = FindObject<UObject>(ANY_PACKAGE, *InternalLink);
+	if (RequiredObject != nullptr)
+	{
+		if(Action && *Action == TEXT("select"))
+		{
+			FContentBrowserModule& ContentBrowserModule = FModuleManager::Get().LoadModuleChecked<FContentBrowserModule>("ContentBrowser");
+			TArray<UObject*> AssetToBrowse;
+			AssetToBrowse.Add(RequiredObject);
+			ContentBrowserModule.Get().SyncBrowserToAssets(AssetToBrowse);
+		}
+		else
+		{
+			FAssetEditorManager::Get().OpenEditorForAsset(RequiredObject);
+		}
+	}
+}
+
+static void OnAssetLinkClicked(const FSlateHyperlinkRun::FMetadata& Metadata)
+{
+	const FString* Url = Metadata.Find(TEXT("href"));
+	const FString* Action = Metadata.Find(TEXT("action"));
+	if(Url)
+	{
+		ParseAssetLink(*Url, Action);
+	}
+}
+
+static TSharedRef<IToolTip> OnGenerateDocTooltip(const FSlateHyperlinkRun::FMetadata& Metadata)
+{
+	FText DisplayText;
+	const FString* Url = Metadata.Find(TEXT("href"));
+	if(Url != nullptr)
+	{
+		DisplayText = FText::Format(LOCTEXT("DocLinkPattern", "View Documentation: {0}"), FText::FromString(*Url));
+	}
+	else
+	{
+		DisplayText = LOCTEXT("UnknownLink", "Empty Hyperlink");
+	}
+
+	FString UrlString;
+	if(Url != nullptr)
+	{
+		UrlString = *Url;
+	}
+
+	// urls for rich tooltips must start with Shared/
+	if(UrlString.Len() > 0 && !UrlString.StartsWith(TEXT("Shared")))
+	{
+		UrlString = FString(TEXT("Shared")) / UrlString;
+	}
+
+	FString ExcerptString;
+	const FString* Excerpt = Metadata.Find(TEXT("excerpt"));
+	if(Excerpt != nullptr)
+	{
+		ExcerptString = *Excerpt;
+	}
+
+	return IDocumentation::Get()->CreateToolTip(DisplayText, nullptr, UrlString, ExcerptString);
+}
+
+static FText OnGetAssetTooltipText(const FSlateHyperlinkRun::FMetadata& Metadata)
+{
+	const FString* Url = Metadata.Find(TEXT("href"));
+	if(Url != nullptr)
+	{
+		const FString* Action = Metadata.Find(TEXT("action"));
+		return FText::Format(LOCTEXT("AssetLinkPattern", "{0} asset: {1}"), (Action == nullptr || *Action == TEXT("select")) ? LOCTEXT("AssetOpenDesc", "Open") : LOCTEXT("AssetFindDesc", "Find"), FText::FromString(*Url));
+	}
+
+	return LOCTEXT("InvalidAssetLink", "Invalid Asset Link");
+}
+
+static FText OnGetCodeTooltipText(const FSlateHyperlinkRun::FMetadata& Metadata)
+{
+	const FString* Url = Metadata.Find(TEXT("href"));
+	if(Url != nullptr)
+	{
+		const bool bUseShortIDEName = true;
+		return FText::Format(LOCTEXT("CodeLinkPattern", "Open code in {0}: {1}"), FSourceCodeNavigation::GetSuggestedSourceCodeIDE(), FText::FromString(*Url));
+	}
+
+	return LOCTEXT("InvalidCodeLink", "Invalid Code Link");
+}
+
+static FText OnGetTutorialTooltipText(const FSlateHyperlinkRun::FMetadata& Metadata)
+{
+	const FString* Url = Metadata.Find(TEXT("href"));
+	if(Url != nullptr)
+	{
+		return FText::Format(LOCTEXT("TutorialLinkPattern", "Open tutorial: {0}"), FText::FromString(*Url));
+	}
+
+	return LOCTEXT("InvalidTutorialLink", "Invalid Tutorial Link");
+}
+
+void FTutorialText::GetRichTextDecorators(TArray< TSharedRef< class ITextDecorator > >& OutDecorators)
+{
+	Initialize();
+
+	for(const auto& HyperlinkDesc : HyperlinkDescs)
+	{
+		OutDecorators.Add(FHyperlinkDecorator::Create(HyperlinkDesc->Id, HyperlinkDesc->OnClickedDelegate, HyperlinkDesc->TooltipTextDelegate, HyperlinkDesc->TooltipDelegate));
+	}
+
+	OutDecorators.Add(FTextStyleDecorator::Create());
+}
+
+void FTutorialText::Initialize()
+{
+	if(HyperlinkDescs.Num() == 0)
+	{
+		HyperlinkDescs.Add(MakeShareable(new FHyperlinkTypeDesc(
+			EHyperlinkType::Browser, 
+			LOCTEXT("BrowserLinkTypeLabel", "URL"), 
+			LOCTEXT("BrowserLinkTypeTooltip", "A link that opens a browser window (e.g. http://www.unrealengine.com)"),
+			TEXT("browser"),
+			FSlateHyperlinkRun::FOnClick::CreateStatic(&OnBrowserLinkClicked))));
+
+		HyperlinkDescs.Add(MakeShareable(new FHyperlinkTypeDesc(
+			EHyperlinkType::UDN, 
+			LOCTEXT("UDNLinkTypeLabel", "UDN"), 
+			LOCTEXT("UDNLinkTypeTooltip", "A link that opens some UDN documentation (e.g. /Engine/Blueprints/UserGuide/Types/ClassBlueprint)"),
+			TEXT("udn"),
+			FSlateHyperlinkRun::FOnClick::CreateStatic(&OnDocLinkClicked), 
+			FSlateHyperlinkRun::FOnGetTooltipText(),
+			FSlateHyperlinkRun::FOnGenerateTooltip::CreateStatic(&OnGenerateDocTooltip))));
+
+		HyperlinkDescs.Add(MakeShareable(new FHyperlinkTypeDesc(
+			EHyperlinkType::Asset, 
+			LOCTEXT("AssetLinkTypeLabel", "Asset"), 
+			LOCTEXT("AssetLinkTypeTooltip", "A link that opens an asset (e.g. /Game/StaticMeshes/SphereMesh.SphereMesh)"),
+			TEXT("asset"),
+			FSlateHyperlinkRun::FOnClick::CreateStatic(&OnAssetLinkClicked),
+			FSlateHyperlinkRun::FOnGetTooltipText::CreateStatic(&OnGetAssetTooltipText))));
+
+		HyperlinkDescs.Add(MakeShareable(new FHyperlinkTypeDesc(
+			EHyperlinkType::Code, 
+			LOCTEXT("CodeLinkTypeLabel", "Code"), 
+			LOCTEXT("CodeLinkTypeTooltip", "A link that opens code in your selected IDE.\nFor example: [PROJECTPATH]/Private/SourceFile.cpp,1,1.\nThe numbers correspond to line number and column number.\nYou can use [PROJECT], [PROJECTPATH] and [ENGINEPATH] tags to make paths."),
+			TEXT("code"),
+			FSlateHyperlinkRun::FOnClick::CreateStatic(&OnCodeLinkClicked),
+			FSlateHyperlinkRun::FOnGetTooltipText::CreateStatic(&OnGetCodeTooltipText))));
+
+		HyperlinkDescs.Add(MakeShareable(new FHyperlinkTypeDesc(
+			EHyperlinkType::Tutorial, 
+			LOCTEXT("TutorialLinkTypeLabel", "Tutorial"), 
+			LOCTEXT("TutorialLinkTypeTooltip", "A type of asset link that opens another tutorial, e.g. /Game/Tutorials/StaticMeshTutorial.StaticMeshTutorial"),
+			TEXT("tutorial"),
+			FSlateHyperlinkRun::FOnClick::CreateStatic(&OnTutorialLinkClicked),
+			FSlateHyperlinkRun::FOnGetTooltipText::CreateStatic(&OnGetTutorialTooltipText))));
+	}
+}
+
+const TArray<TSharedPtr<FHyperlinkTypeDesc>>& FTutorialText::GetHyperLinkDescs()
+{
+	Initialize();
+
+	return HyperlinkDescs;
+}
+
+#undef LOCTEXT_NAMESPACE
\ No newline at end of file
diff --git a/Engine/Source/Editor/IntroTutorials/Private/TutorialText.h b/Engine/Source/Editor/IntroTutorials/Private/TutorialText.h
new file mode 100644
--- /dev/null
+++ b/Engine/Source/Editor/IntroTutorials/Private/TutorialText.h
@@ -0,0 +1,67 @@
+// Copyright 1998-2014 Epic Games, Inc. All Rights Reserved.
+
+#pragma once
+
+namespace EHyperlinkType
+{
+	enum Type
+	{
+		Browser,
+		UDN,
+		Tutorial,
+		Code,
+		Asset
+	};
+}
+
+/** Helper struct to hold info about hyperlink types */
+struct FHyperlinkTypeDesc
+{
+	FHyperlinkTypeDesc(EHyperlinkType::Type InType, const FText& InText, const FText& InTooltipText, const FString& InId, FSlateHyperlinkRun::FOnClick InOnClickedDelegate, FSlateHyperlinkRun::FOnGetTooltipText InTooltipTextDelegate = FSlateHyperlinkRun::FOnGetTooltipText(), FSlateHyperlinkRun::FOnGenerateTooltip InTooltipDelegate = FSlateHyperlinkRun::FOnGenerateTooltip())
+		: Type(InType)
+		, Id(InId)
+		, Text(InText)
+		, TooltipText(InTooltipText)
+		, OnClickedDelegate(InOnClickedDelegate)
+		, TooltipTextDelegate(InTooltipTextDelegate)
+		, TooltipDelegate(InTooltipDelegate)
+	{
+	}
+
+	/** The type of the link */
+	EHyperlinkType::Type Type;
+
+	/** Tag used by this hyperlink's run */
+	FString Id;
+
+	/** Text to display in the UI */
+	FText Text;
+
+	/** Tooltip text to display in the UI */
+	FText TooltipText;
+
+	/** Delegate to execute for this hyperlink's run */
+	FSlateHyperlinkRun::FOnClick OnClickedDelegate;
+
+	/** Delegate used to retrieve the text to display in the hyperlink's tooltip */
+	FSlateHyperlinkRun::FOnGetTooltipText TooltipTextDelegate;
+
+	/** Delegate used to generate hyperlink's tooltip */
+	FSlateHyperlinkRun::FOnGenerateTooltip TooltipDelegate;
+};
+
+
+/** Helper functions for generating rich text */
+struct FTutorialText
+{
+public:
+	static void GetRichTextDecorators(TArray<TSharedRef<class ITextDecorator>>& OutDecorators);
+
+	static const TArray<TSharedPtr<FHyperlinkTypeDesc>>& GetHyperLinkDescs();
+
+private:
+	static void Initialize();
+
+private:
+	static TArray<TSharedPtr<FHyperlinkTypeDesc>> HyperlinkDescs;
+};
\ No newline at end of file
diff --git a/Engine/Source/Runtime/Slate/Private/Framework/Text/SlateHyperlinkRun.cpp b/Engine/Source/Runtime/Slate/Private/Framework/Text/SlateHyperlinkRun.cpp
--- a/Engine/Source/Runtime/Slate/Private/Framework/Text/SlateHyperlinkRun.cpp
+++ b/Engine/Source/Runtime/Slate/Private/Framework/Text/SlateHyperlinkRun.cpp
@@ -56,9 +56,31 @@ int8 FSlateHyperlinkRun::GetKerning( int32 CurrentIndex, float Scale ) const
 
 TSharedRef< ILayoutBlock > FSlateHyperlinkRun::CreateBlock( int32 StartIndex, int32 EndIndex, FVector2D Size, const TSharedPtr< IRunRenderer >& Renderer )
 {
+	FText ToolTipText;
+	TSharedPtr<IToolTip> ToolTip;
+	
+	if(TooltipDelegate.IsBound())
+	{
+		ToolTip = TooltipDelegate.Execute(RunInfo.MetaData);
+	}
+	else
+	{
+		const FString* Url = RunInfo.MetaData.Find(TEXT("href"));
+		if(TooltipTextDelegate.IsBound())
+		{
+			ToolTipText = TooltipTextDelegate.Execute(RunInfo.MetaData);
+		}
+		else if(Url != nullptr)
+		{
+			ToolTipText = FText::FromString(*Url);
+		}
+	}
+
 	TSharedRef< SWidget > Widget = SNew( SRichTextHyperlink, ViewModel )
 		.Style( &Style )
 		.Text( FText::FromString( FString( EndIndex - StartIndex, **Text + StartIndex ) ) )
+		.ToolTip( ToolTip )
+		.ToolTipText( ToolTipText )
 		.OnNavigate( this, &FSlateHyperlinkRun::OnNavigate );
 
 	Children.Add( Widget );
diff --git a/Engine/Source/Runtime/Slate/Private/Framework/Text/TextDecorators.cpp b/Engine/Source/Runtime/Slate/Private/Framework/Text/TextDecorators.cpp
--- a/Engine/Source/Runtime/Slate/Private/Framework/Text/TextDecorators.cpp
+++ b/Engine/Source/Runtime/Slate/Private/Framework/Text/TextDecorators.cpp
@@ -93,14 +93,16 @@ TSharedRef< ISlateRun > FImageDecorator::Create(const TSharedRef<class FTextLayo
 	return FSlateImageRun::Create( RunInfo, InOutModelText, &ImageStyle.Image, ImageStyle.Baseline, ModelRange );
 }
 
-TSharedRef< FHyperlinkDecorator > FHyperlinkDecorator::Create( FString Id, const FSlateHyperlinkRun::FOnClick& NavigateDelegate )
+TSharedRef< FHyperlinkDecorator > FHyperlinkDecorator::Create( FString Id, const FSlateHyperlinkRun::FOnClick& NavigateDelegate, const FSlateHyperlinkRun::FOnGetTooltipText& InToolTipTextDelegate, const FSlateHyperlinkRun::FOnGenerateTooltip& InToolTipDelegate )
 {
-	return MakeShareable( new FHyperlinkDecorator( Id, NavigateDelegate ) );
+	return MakeShareable( new FHyperlinkDecorator( Id, NavigateDelegate, InToolTipTextDelegate, InToolTipDelegate ) );
 }
 
-FHyperlinkDecorator::FHyperlinkDecorator( FString InId, const FSlateHyperlinkRun::FOnClick& InNavigateDelegate )
+FHyperlinkDecorator::FHyperlinkDecorator( FString InId, const FSlateHyperlinkRun::FOnClick& InNavigateDelegate, const FSlateHyperlinkRun::FOnGetTooltipText& InToolTipTextDelegate, const FSlateHyperlinkRun::FOnGenerateTooltip& InToolTipDelegate )
 	: Id( InId )
 	, NavigateDelegate( InNavigateDelegate )
+	, ToolTipTextDelegate( InToolTipTextDelegate )
+	, ToolTipDelegate( InToolTipDelegate )
 {
 
 }
@@ -144,6 +146,10 @@ TSharedRef< ISlateRun > FHyperlinkDecorator::Create(const TSharedRef<class FText
 		RunInfo.MetaData.Add(Pair.Key, OriginalText.Mid( Pair.Value.BeginIndex, Pair.Value.EndIndex - Pair.Value.BeginIndex));
 	}
 
-	return FSlateHyperlinkRun::Create( RunInfo, InOutModelText, Style->GetWidgetStyle<FHyperlinkStyle>( FName( *StyleName ) ), NavigateDelegate, ModelRange );
+	TSharedRef<FSlateHyperlinkRun> HyperlinkRun = FSlateHyperlinkRun::Create( RunInfo, InOutModelText, Style->GetWidgetStyle<FHyperlinkStyle>( FName( *StyleName ) ), NavigateDelegate, ModelRange );
+	HyperlinkRun->OnGenerateToolTip() = ToolTipDelegate;
+	HyperlinkRun->OnGetToolTipText() = ToolTipTextDelegate;
+
+	return HyperlinkRun;
 }
 #endif //WITH_FANCY_TEXT
\ No newline at end of file
diff --git a/Engine/Source/Runtime/Slate/Public/Framework/Text/SlateHyperlinkRun.h b/Engine/Source/Runtime/Slate/Public/Framework/Text/SlateHyperlinkRun.h
--- a/Engine/Source/Runtime/Slate/Public/Framework/Text/SlateHyperlinkRun.h
+++ b/Engine/Source/Runtime/Slate/Public/Framework/Text/SlateHyperlinkRun.h
@@ -8,7 +8,9 @@ class SLATE_API FSlateHyperlinkRun : public ISlateRun, public TSharedFromThis< F
 public:
 
 	typedef TMap< FString, FString > FMetadata;
-	DECLARE_DELEGATE_OneParam( FOnClick, const FMetadata& /*Metadata*/ )
+	DECLARE_DELEGATE_OneParam( FOnClick, const FMetadata& /*Metadata*/ );
+	DECLARE_DELEGATE_RetVal_OneParam( FText, FOnGetTooltipText, const FMetadata& /*Metadata*/ );
+	DECLARE_DELEGATE_RetVal_OneParam( TSharedRef<IToolTip>, FOnGenerateTooltip, const FMetadata& /*Metadata*/ );
 
 public:
 
@@ -73,6 +75,10 @@ public:
 
 	virtual const FRunInfo& GetRunInfo() const override;
 
+	FOnGenerateTooltip& OnGenerateToolTip() { return TooltipDelegate; }
+
+	FOnGetTooltipText& OnGetToolTipText() { return TooltipTextDelegate; }
+
 private:
 
 	FSlateHyperlinkRun( const FRunInfo& InRunInfo, const TSharedRef< const FString >& InText, const FHyperlinkStyle& InStyle, FOnClick InNavigateDelegate );
@@ -92,6 +98,8 @@ private:
 	FTextRange Range;
 	FHyperlinkStyle Style;
 	FOnClick NavigateDelegate;
+	FOnGenerateTooltip TooltipDelegate;
+	FOnGetTooltipText TooltipTextDelegate;
 
 	TSharedRef< FWidgetViewModel > ViewModel;
 	TArray< TSharedRef<SWidget> > Children;
diff --git a/Engine/Source/Runtime/Slate/Public/Framework/Text/TextDecorators.h b/Engine/Source/Runtime/Slate/Public/Framework/Text/TextDecorators.h
--- a/Engine/Source/Runtime/Slate/Public/Framework/Text/TextDecorators.h
+++ b/Engine/Source/Runtime/Slate/Public/Framework/Text/TextDecorators.h
@@ -59,7 +59,7 @@ class SLATE_API FHyperlinkDecorator : public ITextDecorator
 {
 public:
 
-	static TSharedRef< FHyperlinkDecorator > Create( FString Id, const FSlateHyperlinkRun::FOnClick& NavigateDelegate );
+	static TSharedRef< FHyperlinkDecorator > Create( FString Id, const FSlateHyperlinkRun::FOnClick& NavigateDelegate, const FSlateHyperlinkRun::FOnGetTooltipText& InToolTipTextDelegate = FSlateHyperlinkRun::FOnGetTooltipText(), const FSlateHyperlinkRun::FOnGenerateTooltip& InToolTipDelegate = FSlateHyperlinkRun::FOnGenerateTooltip() );
 	virtual ~FHyperlinkDecorator() {}
 
 public:
@@ -70,12 +70,14 @@ public:
 
 private:
 
-	FHyperlinkDecorator( FString InId, const FSlateHyperlinkRun::FOnClick& InNavigateDelegate );
+	FHyperlinkDecorator( FString InId, const FSlateHyperlinkRun::FOnClick& InNavigateDelegate, const FSlateHyperlinkRun::FOnGetTooltipText& InToolTipTextDelegate, const FSlateHyperlinkRun::FOnGenerateTooltip& InToolTipDelegate );
 
 private:
 
 	FString Id;
 	FSlateHyperlinkRun::FOnClick NavigateDelegate;
+	FSlateHyperlinkRun::FOnGetTooltipText ToolTipTextDelegate;
+	FSlateHyperlinkRun::FOnGenerateTooltip ToolTipDelegate;
 };
 
 
