CIS fix - Fixed non-unity build

diff --git a/Engine/Source/Editor/IntroTutorials/Private/STutorialEditableText.h b/Engine/Source/Editor/IntroTutorials/Private/STutorialEditableText.h
--- a/Engine/Source/Editor/IntroTutorials/Private/STutorialEditableText.h
+++ b/Engine/Source/Editor/IntroTutorials/Private/STutorialEditableText.h
@@ -2,118 +2,6 @@
 
 #pragma once
 
-/** Text style and name to display in the UI */
-struct FTextStyleAndName
-{
-	/** Legacy: Flags controlling which TTF or OTF font should be picked from the given font family */
-	struct EFontStyle
-	{
-		typedef uint8 Flags;
-		enum Flag
-		{
-			Regular = 0,
-			Bold = 1<<0,
-			Italic = 1<<1,
-		};
-	};
-
-	FTextStyleAndName(FName InStyle, FText InDisplayName)
-		: Style(InStyle)
-		, DisplayName(InDisplayName)
-	{}
-
-	FRunInfo CreateRunInfo() const
-	{
-		FRunInfo RunInfo(TEXT("TextStyle"));
-		RunInfo.MetaData.Add(TEXT("Style"), Style.ToString());
-		return RunInfo;
-	}
-
-	static FName GetStyleFromRunInfo(const FRunInfo& InRunInfo)
-	{
-		const FString* const StyleString = InRunInfo.MetaData.Find(TEXT("Style"));
-		if(StyleString)
-		{
-			return **StyleString;
-		}
-		else
-		{
-			// legacy styling support - try to map older flexible styles to new fixed styles
-
-			int32 FontSize = 11;
-			const FString* const FontSizeString = InRunInfo.MetaData.Find(TEXT("FontSize"));
-			if(FontSizeString)
-			{
-				FontSize = static_cast<uint8>(FPlatformString::Atoi(**FontSizeString));
-			}
-
-			if(FontSize > 24)
-			{
-				return TEXT("Tutorials.Content.HeaderText2");
-			}
-			else if(FontSize > 11 && FontSize <= 24)
-			{
-				return TEXT("Tutorials.Content.HeaderText1");
-			}
-			else
-			{
-				EFontStyle::Flags FontStyle = EFontStyle::Regular;
-				const FString* const FontStyleString = InRunInfo.MetaData.Find(TEXT("FontStyle"));
-				if(FontStyleString)
-				{
-					if(*FontStyleString == TEXT("Bold"))
-					{
-						FontStyle = EFontStyle::Bold;
-					}
-					else if(*FontStyleString == TEXT("Italic"))
-					{
-						FontStyle = EFontStyle::Italic;
-					}
-					else if(*FontStyleString == TEXT("BoldItalic"))
-					{
-						FontStyle = EFontStyle::Bold | EFontStyle::Italic;
-					}
-				}
-
-				FLinearColor FontColor = FLinearColor::Black;
-				const FString* const FontColorString = InRunInfo.MetaData.Find(TEXT("FontColor"));
-				if(FontColorString && !FontColor.InitFromString(*FontColorString))
-				{
-					FontColor = FLinearColor::Black;
-				}
-
-				if(FontStyle != EFontStyle::Regular || FontColor != FLinearColor::Black)
-				{
-					return TEXT("Tutorials.Content.TextBold");
-				}
-				else
-				{
-					return TEXT("Tutorials.Content.Text");
-				}
-			}
-		}
-
-		return NAME_None;
-	}
-
-	FTextBlockStyle CreateTextBlockStyle() const
-	{
-		return FEditorStyle::Get().GetWidgetStyle<FTextBlockStyle>(Style);
-	}
-
-	static FTextBlockStyle CreateTextBlockStyle(const FRunInfo& InRunInfo)
-	{
-		return FEditorStyle::Get().GetWidgetStyle<FTextBlockStyle>(GetStyleFromRunInfo(InRunInfo));
-	}
-
-	/** The style identifier */
-	FName Style;
-
-	/** The text to display for this style */
-	FText DisplayName;
-};
-
-
 class STutorialEditableText : public SCompoundWidget
 {
 public:
@@ -148,9 +36,9 @@ protected:
 
 	FText GetActiveStyleName() const;
 
-	void OnActiveStyleChanged(TSharedPtr<FTextStyleAndName> NewValue, ESelectInfo::Type);
+	void OnActiveStyleChanged(TSharedPtr<struct FTextStyleAndName> NewValue, ESelectInfo::Type);
 
-	TSharedRef<SWidget> GenerateStyleComboEntry(TSharedPtr<FTextStyleAndName> SourceEntry);
+	TSharedRef<SWidget> GenerateStyleComboEntry(TSharedPtr<struct FTextStyleAndName> SourceEntry);
 
 	void StyleSelectedText();
 
@@ -192,15 +80,15 @@ protected:
 	FSlateHyperlinkRun::FOnClick OnAssetLinkClicked;
 
 	TSharedPtr<SComboButton> HyperlinkComboButton;
-	TSharedPtr<SComboBox<TSharedPtr<FTextStyleAndName>>> FontComboBox;
+	TSharedPtr<SComboBox<TSharedPtr<struct FTextStyleAndName>>> FontComboBox;
 	TSharedPtr<STextBlock> HyperlinkNameTextBlock;
 	TSharedPtr<SEditableTextBox> HyperlinkURLTextBox;
 	TSharedPtr<SEditableTextBox> UDNExcerptTextBox;
 
-	TSharedPtr<FTextStyleAndName> ActiveStyle;
-	TSharedPtr<FTextStyleAndName> HyperlinkStyle;
+	TSharedPtr<struct FTextStyleAndName> ActiveStyle;
+	TSharedPtr<struct FTextStyleAndName> HyperlinkStyle;
 
-	TArray<TSharedPtr<FTextStyleAndName>> StylesAndNames;
+	TArray<TSharedPtr<struct FTextStyleAndName>> StylesAndNames;
 
 	FOnTextCommitted OnTextCommitted;
 	FOnTextChanged OnTextChanged;
diff --git a/Engine/Source/Editor/IntroTutorials/Private/TutorialText.h b/Engine/Source/Editor/IntroTutorials/Private/TutorialText.h
--- a/Engine/Source/Editor/IntroTutorials/Private/TutorialText.h
+++ b/Engine/Source/Editor/IntroTutorials/Private/TutorialText.h
@@ -51,6 +51,119 @@ struct FHyperlinkTypeDesc
 };
 
 
+/** Text style and name to display in the UI */
+struct FTextStyleAndName
+{
+	/** Legacy: Flags controlling which TTF or OTF font should be picked from the given font family */
+	struct EFontStyle
+	{
+		typedef uint8 Flags;
+		enum Flag
+		{
+			Regular = 0,
+			Bold = 1<<0,
+			Italic = 1<<1,
+		};
+	};
+
+	FTextStyleAndName(FName InStyle, FText InDisplayName)
+		: Style(InStyle)
+		, DisplayName(InDisplayName)
+	{}
+
+	FRunInfo CreateRunInfo() const
+	{
+		FRunInfo RunInfo(TEXT("TextStyle"));
+		RunInfo.MetaData.Add(TEXT("Style"), Style.ToString());
+		return RunInfo;
+	}
+
+	static FName GetStyleFromRunInfo(const FRunInfo& InRunInfo)
+	{
+		const FString* const StyleString = InRunInfo.MetaData.Find(TEXT("Style"));
+		if(StyleString)
+		{
+			return **StyleString;
+		}
+		else
+		{
+			// legacy styling support - try to map older flexible styles to new fixed styles
+
+			int32 FontSize = 11;
+			const FString* const FontSizeString = InRunInfo.MetaData.Find(TEXT("FontSize"));
+			if(FontSizeString)
+			{
+				FontSize = static_cast<uint8>(FPlatformString::Atoi(**FontSizeString));
+			}
+
+			if(FontSize > 24)
+			{
+				return TEXT("Tutorials.Content.HeaderText2");
+			}
+			else if(FontSize > 11 && FontSize <= 24)
+			{
+				return TEXT("Tutorials.Content.HeaderText1");
+			}
+			else
+			{
+				EFontStyle::Flags FontStyle = EFontStyle::Regular;
+				const FString* const FontStyleString = InRunInfo.MetaData.Find(TEXT("FontStyle"));
+				if(FontStyleString)
+				{
+					if(*FontStyleString == TEXT("Bold"))
+					{
+						FontStyle = EFontStyle::Bold;
+					}
+					else if(*FontStyleString == TEXT("Italic"))
+					{
+						FontStyle = EFontStyle::Italic;
+					}
+					else if(*FontStyleString == TEXT("BoldItalic"))
+					{
+						FontStyle = EFontStyle::Bold | EFontStyle::Italic;
+					}
+				}
+
+				FLinearColor FontColor = FLinearColor::Black;
+				const FString* const FontColorString = InRunInfo.MetaData.Find(TEXT("FontColor"));
+				if(FontColorString && !FontColor.InitFromString(*FontColorString))
+				{
+					FontColor = FLinearColor::Black;
+				}
+
+				if(FontStyle != EFontStyle::Regular || FontColor != FLinearColor::Black)
+				{
+					return TEXT("Tutorials.Content.TextBold");
+				}
+				else
+				{
+					return TEXT("Tutorials.Content.Text");
+				}
+			}
+		}
+
+		return NAME_None;
+	}
+
+	FTextBlockStyle CreateTextBlockStyle() const
+	{
+		return FEditorStyle::Get().GetWidgetStyle<FTextBlockStyle>(Style);
+	}
+
+	static FTextBlockStyle CreateTextBlockStyle(const FRunInfo& InRunInfo)
+	{
+		return FEditorStyle::Get().GetWidgetStyle<FTextBlockStyle>(GetStyleFromRunInfo(InRunInfo));
+	}
+
+	/** The style identifier */
+	FName Style;
+
+	/** The text to display for this style */
+	FText DisplayName;
+};
+
+
+
 /** Helper functions for generating rich text */
 struct FTutorialText
 {
