Fixed visibility of FX widget & tag for tutorial button

Prevented empty tabs launching when tutorial button anchored-tutorials are run.
Prevented empty tutorials from blocking input.

diff --git a/Engine/Source/Editor/IntroTutorials/Private/STutorialButton.cpp b/Engine/Source/Editor/IntroTutorials/Private/STutorialButton.cpp
--- a/Engine/Source/Editor/IntroTutorials/Private/STutorialButton.cpp
+++ b/Engine/Source/Editor/IntroTutorials/Private/STutorialButton.cpp
@@ -27,11 +27,10 @@ void STutorialButton::Construct(const FArguments& InArgs)
 	PulseAnimation.AddCurve(0.0f, TutorialButtonConstants::PulseAnimationLength, ECurveEaseFunction::Linear);
 	PulseAnimation.Play();
 
-	FTutorialMetaData TagMeta(*FString::Printf(TEXT("%s.TutorialLaunchButton"), *Context.ToString()), TEXT("TutorialLaunchButton"));
 	ChildSlot
 	[
 		SNew(SButton)
-		.AddMetaData<FTutorialMetaData>(TagMeta)
+		.AddMetaData<FTagMetaData>(*FString::Printf(TEXT("%s.TutorialLaunchButton"), *Context.ToString()))
 		.ButtonStyle(FEditorStyle::Get(), "TutorialLaunch.Button")
 		.ToolTipText(this, &STutorialButton::GetButtonToolTip)
 		.OnClicked(this, &STutorialButton::HandleButtonClicked)
diff --git a/Engine/Source/Editor/IntroTutorials/Private/STutorialContent.cpp b/Engine/Source/Editor/IntroTutorials/Private/STutorialContent.cpp
--- a/Engine/Source/Editor/IntroTutorials/Private/STutorialContent.cpp
+++ b/Engine/Source/Editor/IntroTutorials/Private/STutorialContent.cpp
@@ -53,6 +53,7 @@ void STutorialContent::Construct(const FArguments& InArgs, const FTutorialConten
 	ChildSlot
 	[
 		SNew(SFxWidget)
+		.Visibility(EVisibility::SelfHitTestInvisible)
 		.RenderScale(this, &STutorialContent::GetAnimatedZoom)
 		.RenderScaleOrigin(FVector2D(0.5f, 0.5f))
 		[
diff --git a/Engine/Source/Editor/IntroTutorials/Private/TutorialStructCustomization.cpp b/Engine/Source/Editor/IntroTutorials/Private/TutorialStructCustomization.cpp
--- a/Engine/Source/Editor/IntroTutorials/Private/TutorialStructCustomization.cpp
+++ b/Engine/Source/Editor/IntroTutorials/Private/TutorialStructCustomization.cpp
@@ -265,9 +265,9 @@ private:
 			{
 				if (MetaDataEntry->IsOfType<FTutorialMetaData>())
 				{
-					TSharedRef<FTutorialMetaData> GraphNodeMeta = StaticCastSharedRef<FGraphNodeMetaData>(MetaDataEntry);
-					StructPropertyHandle->GetChildHandle(GET_MEMBER_NAME_CHECKED(FTutorialContentAnchor, FriendlyName))->SetValue(GraphNodeMeta->FriendlyName);
-					StructPropertyHandle->GetChildHandle(GET_MEMBER_NAME_CHECKED(FTutorialContentAnchor, TabTypeToOpen))->SetValue(GraphNodeMeta->TabTypeToOpen);
+					TSharedRef<FTutorialMetaData> TutorialMeta = StaticCastSharedRef<FTutorialMetaData>(MetaDataEntry);
+					StructPropertyHandle->GetChildHandle(GET_MEMBER_NAME_CHECKED(FTutorialContentAnchor, FriendlyName))->SetValue(TutorialMeta->FriendlyName);
+					StructPropertyHandle->GetChildHandle(GET_MEMBER_NAME_CHECKED(FTutorialContentAnchor, TabTypeToOpen))->SetValue(TutorialMeta->TabTypeToOpen);
 				}
 				else if (MetaDataEntry->IsOfType<FGraphNodeMetaData>())
 				{
