New tutorials are now enabled by default

Removed command-line option and older tutorials

diff --git a/Engine/Source/Editor/IntroTutorials/Private/IntroTutorials.cpp b/Engine/Source/Editor/IntroTutorials/Private/IntroTutorials.cpp
--- a/Engine/Source/Editor/IntroTutorials/Private/IntroTutorials.cpp
+++ b/Engine/Source/Editor/IntroTutorials/Private/IntroTutorials.cpp
@@ -345,16 +345,8 @@ void FIntroTutorials::MainFrameLoad(TSharedPtr<SWindow> InRootWindow, bool bIsNe
 
 void FIntroTutorials::SummonTutorialHome()
 {
-	if(FParse::Param(FCommandLine::Get(), TEXT("NewTutorials")))
-	{
-		FLevelEditorModule& LevelEditorModule = FModuleManager::LoadModuleChecked<FLevelEditorModule>( "LevelEditor" );
-		SummonTutorialBrowser(LevelEditorModule.GetLevelEditorTab()->GetParentWindow().ToSharedRef());
-	}
-	else
-	{
-		CurrentObjectClass = nullptr;
-		SummonTutorialWindowForPage(HomePath);
-	}
+	FLevelEditorModule& LevelEditorModule = FModuleManager::LoadModuleChecked<FLevelEditorModule>( "LevelEditor" );
+	SummonTutorialBrowser(LevelEditorModule.GetLevelEditorTab()->GetParentWindow().ToSharedRef());
 }
 
 void FIntroTutorials::SummonBlueprintTutorialHome(UObject* Asset, bool bForceWelcome)
@@ -563,54 +555,26 @@ FWelcomeTutorialProperties const* FIntroTutorials::FindAssetEditorTutorialProper
 
 bool FIntroTutorials::MaybeOpenWelcomeTutorial(const FString& TutorialPath, const FString& ConfigSettingName)
 {
-	if(FParse::Param(FCommandLine::Get(), TEXT("NewTutorials")))
+	// try editor startup tutorial
+	TSubclassOf<UEditorTutorial> EditorStartupTutorialClass = LoadClass<UEditorTutorial>(NULL, *GetDefault<UEditorTutorialSettings>()->StartupTutorial.AssetLongPathname, NULL, LOAD_None, NULL);
+	if(EditorStartupTutorialClass != nullptr)
 	{
-		// try editor startup tutorial
-		TSubclassOf<UEditorTutorial> EditorStartupTutorialClass = LoadClass<UEditorTutorial>(NULL, *GetDefault<UEditorTutorialSettings>()->StartupTutorial.AssetLongPathname, NULL, LOAD_None, NULL);
-		if(EditorStartupTutorialClass != nullptr)
-		{
-			UEditorTutorial* Tutorial = EditorStartupTutorialClass->GetDefaultObject<UEditorTutorial>();
-			if (!GetDefault<UTutorialStateSettings>()->HaveSeenTutorial(Tutorial))
-			{
-				LaunchTutorial(Tutorial);
-				return true;
-			}
-		}
-
-		// Try project startup tutorial
-		TSubclassOf<UEditorTutorial> ProjectStartupTutorialClass = LoadClass<UEditorTutorial>(NULL, *GetDefault<UTutorialSettings>()->StartupTutorial.AssetLongPathname, NULL, LOAD_None, NULL);
-		if(ProjectStartupTutorialClass != nullptr)
+		UEditorTutorial* Tutorial = EditorStartupTutorialClass->GetDefaultObject<UEditorTutorial>();
+		if (!GetDefault<UTutorialStateSettings>()->HaveSeenTutorial(Tutorial))
 		{
-			UEditorTutorial* Tutorial = ProjectStartupTutorialClass->GetDefaultObject<UEditorTutorial>();
-			if (!GetDefault<UTutorialStateSettings>()->HaveSeenTutorial(Tutorial))
-			{
-				LaunchTutorial(Tutorial);
-				return true;
-			}
+			LaunchTutorial(Tutorial);
+			return true;
 		}
-
-		return false;
-	}
-
-	// don't open if viewing any tutorial other than the index
-	if (TutorialWidget.IsValid() && (TutorialWidget.Pin()->GetCurrentPagePath() != HomePath))
-	{
-		return false;
 	}
 
-	bool bSeenWelcome = false;
-	GConfig->GetBool(*IntroTutorialConfigSection, *ConfigSettingName, bSeenWelcome, GEditorGameAgnosticIni);
-
-	if (!bSeenWelcome)
+	// Try project startup tutorial
+	TSubclassOf<UEditorTutorial> ProjectStartupTutorialClass = LoadClass<UEditorTutorial>(NULL, *GetDefault<UTutorialSettings>()->StartupTutorial.AssetLongPathname, NULL, LOAD_None, NULL);
+	if(ProjectStartupTutorialClass != nullptr)
 	{
-		bool const bPageExists = IDocumentation::Get()->PageExists(*TutorialPath);
-		if (bPageExists)
+		UEditorTutorial* Tutorial = ProjectStartupTutorialClass->GetDefaultObject<UEditorTutorial>();
+		if (!GetDefault<UTutorialStateSettings>()->HaveSeenTutorial(Tutorial))
 		{
-			SummonTutorialWindowForPage(*TutorialPath);
-
-			// Tell ini file that we've seen this now
-			GConfig->SetBool(*IntroTutorialConfigSection, *ConfigSettingName, true, GEditorGameAgnosticIni);
-
+			LaunchTutorial(Tutorial);
 			return true;
 		}
 	}
