Made tutorial title AssetRegistrySearchable.

diff --git a/Engine/Source/Editor/IntroTutorials/Public/EditorTutorial.h b/Engine/Source/Editor/IntroTutorials/Public/EditorTutorial.h
--- a/Engine/Source/Editor/IntroTutorials/Public/EditorTutorial.h
+++ b/Engine/Source/Editor/IntroTutorials/Public/EditorTutorial.h
@@ -205,7 +205,7 @@ class INTROTUTORIALS_API UEditorTutorial : public UObject
 	GENERATED_UCLASS_BODY()
 
 	/** Title of this tutorial, used when presented to the user */
-	UPROPERTY(EditAnywhere, Category="Tutorial")
+	UPROPERTY(EditAnywhere, Category="Tutorial", AssetRegistrySearchable)
 	FText Title;
 
 	/** Icon name for this tutorial, used when presented to the user in the tutorial browser. This is a name for the icon in the Slate editor style. Only used if there isn't a valid texture to use. */
