Asset links in tutorials now open correctly.

Previously they were only opening if the asset had been previously loaded.

diff --git a/Engine/Source/Editor/IntroTutorials/Private/TutorialText.cpp b/Engine/Source/Editor/IntroTutorials/Private/TutorialText.cpp
--- a/Engine/Source/Editor/IntroTutorials/Private/TutorialText.cpp
+++ b/Engine/Source/Editor/IntroTutorials/Private/TutorialText.cpp
@@ -158,7 +158,7 @@ static void OnCodeLinkClicked(const FSlateHyperlinkRun::FMetadata& Metadata)
 
 static void ParseAssetLink(const FString& InternalLink, const FString* Action)
 {
-	UObject* RequiredObject = FindObject<UObject>(ANY_PACKAGE, *InternalLink);
+	UObject* RequiredObject = LoadObject<UObject>(nullptr, *InternalLink);
 	if (RequiredObject != nullptr)
 	{
 		if(Action && *Action == TEXT("select"))
