Reworking navigation into content

Disabled floating navigation widget (but have not removed it yet).
Built navigation controls into the content widget.

diff --git a/Engine/Source/Editor/EditorStyle/Private/SlateEditorStyle.cpp b/Engine/Source/Editor/EditorStyle/Private/SlateEditorStyle.cpp
--- a/Engine/Source/Editor/EditorStyle/Private/SlateEditorStyle.cpp
+++ b/Engine/Source/Editor/EditorStyle/Private/SlateEditorStyle.cpp
@@ -2455,13 +2455,16 @@ void FSlateEditorStyle::FStyle::SetupTutorialStyles()
 			.SetPressedPadding( FMargin(0,1,0,0)));
 
 		Set( "Tutorials.Content.Button", FButtonStyle()
-			.SetNormal( BOX_BRUSH( "Common/ButtonHoverHint", FMargin(4/16.0f), FLinearColor(1,1,1,0.15f) ) )
+			.SetNormal( BOX_BRUSH( "Common/ButtonHoverHint", FMargin(4/16.0f), FLinearColor(0,0,0,0) ) )
 			.SetHovered( BOX_BRUSH( "Common/ButtonHoverHint", FMargin(4/16.0f), FLinearColor(1,1,1,1) ) )
 			.SetPressed( BOX_BRUSH( "Common/ButtonHoverHint", FMargin(4/16.0f), FLinearColor(1,1,1,1) ) )
 			.SetDisabled( BOX_BRUSH( "Common/ButtonHoverHint", FMargin(4/16.0f), FLinearColor(1,1,1,1) ) )
 			.SetNormalPadding( FMargin(0,0,0,1))
 			.SetPressedPadding( FMargin(0,1,0,0)));
 
+		Set( "Tutorials.Content.Color", FLinearColor(1.0f,1.0f,1.0f,0.9f) );
+		Set( "Tutorials.Content.Color.Hovered", FLinearColor(1.0f,1.0f,1.0f,1.0f) );
+
 		Set( "Tutorials.Browser.CategoryArrow", new IMAGE_BRUSH( "Tutorials/BrowserCategoryArrow", FVector2D(24.0f, 24.0f), FLinearColor(0.0f, 0.0f, 0.0f, 1.0f) ) );
 		Set( "Tutorials.Browser.DefaultTutorialIcon", new IMAGE_BRUSH( "Tutorials/DefaultTutorialIcon_40x", FVector2D(40.0f, 40.0f), FLinearColor::White ) );
 		Set( "Tutorials.Browser.DefaultCategoryIcon", new IMAGE_BRUSH( "Tutorials/DefaultCategoryIcon_40x", FVector2D(40.0f, 40.0f), FLinearColor::White ) );
@@ -2469,6 +2472,18 @@ void FSlateEditorStyle::FStyle::SetupTutorialStyles()
 		Set( "Tutorials.Browser.BackButton", new IMAGE_BRUSH( "Tutorials/BrowserBack", FVector2D(32.0f, 32.0f), FLinearColor(0.0f, 0.0f, 0.0f, 1.0f) ) );
 		Set( "Tutorials.Browser.RestartButton", new IMAGE_BRUSH( "Tutorials/BrowserRestart", FVector2D(16.0f, 16.0f), FLinearColor(0.0f, 0.0f, 0.0f, 1.0f) ) );
 
+		Set( "Tutorials.Navigation.Button", FButtonStyle()
+			.SetNormal( BOX_BRUSH( "Tutorials/ButtonHoverHint", FMargin(4/16.0f), FLinearColor(0,0,0,0) ) )
+			.SetHovered( BOX_BRUSH( "Tutorials/ButtonHoverHint", FMargin(4/16.0f), FLinearColor(0,0,0,0) ) )
+			.SetPressed( BOX_BRUSH( "Tutorials/ButtonHoverHint", FMargin(4/16.0f), FLinearColor(0,0,0,0) ) )
+			.SetDisabled( BOX_BRUSH( "Tutorials/ButtonHoverHint", FMargin(4/16.0f), FLinearColor(0,0,0,0) ) )
+			.SetNormalPadding( FMargin(0,0,0,1))
+			.SetPressedPadding( FMargin(0,1,0,0)));
+
+		Set( "Tutorials.Navigation.NextButton", new IMAGE_BRUSH( "Tutorials/NavigationNext", Icon32x32 ) );
+		Set( "Tutorials.Navigation.HomeButton", new IMAGE_BRUSH( "Tutorials/NavigationHome", Icon32x32 ) );
+		Set( "Tutorials.Navigation.BackButton", new IMAGE_BRUSH( "Tutorials/NavigationBack", Icon32x32 ) );
+
 		Set( "Tutorials.WidgetContent", FTextBlockStyle(TutorialText)
 			.SetFontSize(10));
 
diff --git a/Engine/Source/Editor/IntroTutorials/Private/SEditorTutorials.cpp b/Engine/Source/Editor/IntroTutorials/Private/SEditorTutorials.cpp
--- a/Engine/Source/Editor/IntroTutorials/Private/SEditorTutorials.cpp
+++ b/Engine/Source/Editor/IntroTutorials/Private/SEditorTutorials.cpp
@@ -114,9 +114,7 @@ EVisibility SEditorTutorials::GetBrowserVisibility() const
 
 EVisibility SEditorTutorials::GetNavigationVisibility() const
 {
-	UEditorTutorial* CurrentTutorial = OnGetCurrentTutorial.Execute();
-
-	return bShowNavigation && CurrentTutorial != nullptr && !CurrentTutorial->bIsStandalone ? EVisibility::Visible : EVisibility::Collapsed;
+	return EVisibility::Collapsed;
 }
 
 void SEditorTutorials::HandleCloseClicked()
@@ -189,11 +187,17 @@ void SEditorTutorials::RebuildCurrentContent()
 	{
 		ContentBox->AddSlot()
 		[
-			SAssignNew(OverlayContent, STutorialOverlay, &CurrentTutorial->Stages[CurrentTutorialStage])
+			SAssignNew(OverlayContent, STutorialOverlay, CurrentTutorial, &CurrentTutorial->Stages[CurrentTutorialStage])
 			.OnClosed(FSimpleDelegate::CreateSP(this, &SEditorTutorials::HandleCloseClicked))
 			.IsStandalone(CurrentTutorial->bIsStandalone)
 			.ParentWindow(ParentWindow)
 			.AllowNonWidgetContent(bShowNavigation)
+			.OnBackClicked(FSimpleDelegate::CreateSP(this, &SEditorTutorials::HandleBackClicked))
+			.OnHomeClicked(FSimpleDelegate::CreateSP(this, &SEditorTutorials::HandleHomeClicked))
+			.OnNextClicked(FSimpleDelegate::CreateSP(this, &SEditorTutorials::HandleNextClicked))
+			.IsBackEnabled(this, &SEditorTutorials::IsBackButtonEnabled)
+			.IsHomeEnabled(this, &SEditorTutorials::IsHomeButtonEnabled)
+			.IsNextEnabled(this, &SEditorTutorials::IsNextButtonEnabled)
 		];
 	}
 	else
@@ -201,7 +205,7 @@ void SEditorTutorials::RebuildCurrentContent()
 		// create 'empty' overlay, as we may need this for picking visualization
 		ContentBox->AddSlot()
 		[
-			SAssignNew(OverlayContent, STutorialOverlay, nullptr)
+			SAssignNew(OverlayContent, STutorialOverlay, nullptr, nullptr)
 			.OnClosed(FSimpleDelegate::CreateSP(this, &SEditorTutorials::HandleCloseClicked))
 			.IsStandalone(false)
 			.ParentWindow(ParentWindow)
diff --git a/Engine/Source/Editor/IntroTutorials/Private/STutorialContent.cpp b/Engine/Source/Editor/IntroTutorials/Private/STutorialContent.cpp
--- a/Engine/Source/Editor/IntroTutorials/Private/STutorialContent.cpp
+++ b/Engine/Source/Editor/IntroTutorials/Private/STutorialContent.cpp
@@ -16,20 +16,29 @@ namespace TutorialConstants
 	const float MinBorderOpacity = 0.1f;
 	const float ShadowScale = 8.0f;
 	const float MaxBorderOffset = 8.0f;
-	const float BorderSize = 24.0f;
+	const FMargin BorderSizeStandalone(24.0f, 24.0f);
+	const FMargin BorderSize(24.0f, 24.0f, 24.0f, 42.0f);
 }
 
 const float ContentOffset = 10.0f;
 
-void STutorialContent::Construct(const FArguments& InArgs, const FTutorialContent& InContent)
+void STutorialContent::Construct(const FArguments& InArgs, UEditorTutorial* InTutorial, const FTutorialContent& InContent)
 {
 	bIsVisible = Anchor.Type == ETutorialAnchorIdentifier::None;
 
+	Tutorial = InTutorial;
+
 	VerticalAlignment = InArgs._VAlign;
 	HorizontalAlignment = InArgs._HAlign;
 	WidgetOffset = InArgs._Offset;
 	bIsStandalone = InArgs._IsStandalone;
 	OnClosed = InArgs._OnClosed;
+	OnNextClicked = InArgs._OnNextClicked;
+	OnHomeClicked = InArgs._OnHomeClicked;
+	OnBackClicked = InArgs._OnBackClicked;
+	IsBackEnabled = InArgs._IsBackEnabled;
+	IsHomeEnabled = InArgs._IsHomeEnabled;
+	IsNextEnabled = InArgs._IsNextEnabled;
 	Anchor = InArgs._Anchor;
 
 	BorderIntroAnimation.AddCurve(0.0f, TutorialConstants::BorderIntroAnimationLength, ECurveEaseFunction::CubicOut);
@@ -57,45 +66,95 @@ void STutorialContent::Construct(const FArguments& InArgs, const FTutorialConten
 		.RenderScale(this, &STutorialContent::GetAnimatedZoom)
 		.RenderScaleOrigin(FVector2D(0.5f, 0.5f))
 		[
-			SAssignNew(ContentWidget, SBorder)
-
-			// Add more padding if the content is to be displayed centrally (i.e. not on a widget)
-			.Padding(TutorialConstants::BorderSize)
-			.BorderImage(FEditorStyle::GetBrush("Tutorials.Border"))
-			.Visibility(this, &STutorialContent::GetVisibility)
-			.BorderBackgroundColor(this, &STutorialContent::GetBackgroundColor)
-			.ForegroundColor(FCoreStyle::Get().GetSlateColor("InvertedForeground"))
+			SNew(SOverlay)
+			+SOverlay::Slot()
 			[
-				SNew(SFxWidget)
-				.RenderScale(this, &STutorialContent::GetInverseAnimatedZoom)
-				.RenderScaleOrigin(FVector2D(0.5f, 0.5f))
+				SAssignNew(ContentWidget, SBorder)
+
+				// Add more padding if the content is to be displayed centrally (i.e. not on a widget)
+				.Padding(bIsStandalone ? TutorialConstants::BorderSizeStandalone : TutorialConstants::BorderSize)
+				.BorderImage(FEditorStyle::GetBrush("Tutorials.Border"))
+				.Visibility(this, &STutorialContent::GetVisibility)
+				.BorderBackgroundColor(this, &STutorialContent::GetBackgroundColor)
+				.ForegroundColor(FCoreStyle::Get().GetSlateColor("InvertedForeground"))
 				[
-					SNew(SHorizontalBox)
-					+SHorizontalBox::Slot()
-					.AutoWidth()
-					.MaxWidth(600.0f)
-					.VAlign(VAlign_Center)
-					[
-						GenerateContentWidget(InContent, InArgs._WrapTextAt, DocumentationPage)
-					]
-					+SHorizontalBox::Slot()
-					.AutoWidth()
-					.VAlign(VAlign_Top)
-					.Padding(5.0f, 0.0f, 0.0f, 0.0f)
+					SNew(SFxWidget)
+					.RenderScale(this, &STutorialContent::GetInverseAnimatedZoom)
+					.RenderScaleOrigin(FVector2D(0.5f, 0.5f))
 					[
-						SNew(SButton)
-						.OnClicked(this, &STutorialContent::OnCloseButtonClicked)
-						.Visibility(this, &STutorialContent::GetCloseButtonVisibility)
-						.ButtonStyle(&FEditorStyle::Get().GetWidgetStyle<FButtonStyle>("Tutorials.Content.Button"))
-						.ContentPadding(0.0f)
+						SNew(SVerticalBox)
+						+SVerticalBox::Slot()
+						.AutoHeight()
 						[
-							SNew(SImage)
-							.Image(FEditorStyle::GetBrush("Symbols.X"))
-							.ColorAndOpacity(FLinearColor(0.0f, 0.0f, 0.0f, 1.0f))
+							SNew(SHorizontalBox)
+							+SHorizontalBox::Slot()
+							.AutoWidth()
+							.MaxWidth(600.0f)
+							.VAlign(VAlign_Center)
+							[
+								GenerateContentWidget(InContent, InArgs._WrapTextAt, DocumentationPage)
+							]
 						]
 					]
 				]
 			]
+			+SOverlay::Slot()
+			.VAlign(VAlign_Top)
+			.HAlign(HAlign_Right)
+			.Padding(16.0f)
+			[
+				SNew(SHorizontalBox)
+				+SHorizontalBox::Slot()
+				.AutoWidth()
+				.VAlign(VAlign_Top)
+				.Padding(2.0f)
+				[
+					SNew(SComboButton)
+					.ToolTipText(LOCTEXT("MoreOptionsTooltip", "More Options"))
+					.Visibility(this, &STutorialContent::GetMenuButtonVisibility)
+					.ButtonStyle(&FEditorStyle::Get().GetWidgetStyle<FButtonStyle>("Tutorials.Content.Button"))
+					.ContentPadding(0.0f)
+					.OnGetMenuContent(FOnGetContent::CreateSP(this, &STutorialContent::HandleGetMenuContent))
+				]
+				+SHorizontalBox::Slot()
+				.AutoWidth()
+				.VAlign(VAlign_Top)
+				.Padding(0.0f)
+				[
+					SNew(SButton)
+					.ToolTipText(LOCTEXT("QuitStandaloneTooltip", "Close this Message"))
+					.OnClicked(this, &STutorialContent::OnCloseButtonClicked)
+					.Visibility(this, &STutorialContent::GetCloseButtonVisibility)
+					.ButtonStyle(&FEditorStyle::Get().GetWidgetStyle<FButtonStyle>("Tutorials.Content.Button"))
+					.ContentPadding(0.0f)
+					[
+						SNew(SImage)
+						.Image(FEditorStyle::GetBrush("Symbols.X"))
+						.ColorAndOpacity(FLinearColor(0.0f, 0.0f, 0.0f, 1.0f))
+					]
+				]
+			]
+			+SOverlay::Slot()
+			.VAlign(VAlign_Bottom)
+			.HAlign(HAlign_Right)
+			.Padding(8.0f)
+			[
+				SAssignNew(NextButton, SButton)
+				.ToolTipText(this, &STutorialContent::GetNextButtonTooltip)
+				.OnClicked(this, &STutorialContent::HandleNextClicked)
+				.Visibility(this, &STutorialContent::GetMenuButtonVisibility)
+				.ButtonStyle(&FEditorStyle::Get().GetWidgetStyle<FButtonStyle>("Tutorials.Navigation.Button"))
+				.ContentPadding(0.0f)
+				[
+					SNew(SBox)
+					.Padding(8.0f)
+					[
+						SNew(SImage)
+						.Image(this, &STutorialContent::GetNextButtonBrush)
+						.ColorAndOpacity(this, &STutorialContent::GetNextButtonColor)
+					]
+				]
+			]
 		]
 	];
 }
@@ -322,6 +381,11 @@ EVisibility STutorialContent::GetCloseButtonVisibility() const
 	return bIsStandalone ? EVisibility::Visible : EVisibility::Collapsed;
 }
 
+EVisibility STutorialContent::GetMenuButtonVisibility() const
+{
+	return !bIsStandalone ? EVisibility::Visible : EVisibility::Collapsed;
+}
+
 void STutorialContent::HandlePaintNamedWidget(TSharedRef<SWidget> InWidget, const FGeometry& InGeometry)
 {
 	switch(Anchor.Type)
@@ -359,12 +423,12 @@ FSlateColor STutorialContent::GetBackgroundColor() const
 {
 	// note cant use IsHovered() here because our widget is SelfHitTestInvisible
 	const FVector2D CursorPos = FSlateApplication::Get().GetCursorPos();
-	return CachedContentGeometry.IsUnderLocation(CursorPos) ? FLinearColor(1.0f, 1.0f, 1.0f, 1.0f) : FLinearColor(1.0f, 1.0f, 1.0f, 0.8f);
+	return CachedContentGeometry.IsUnderLocation(CursorPos) ? FEditorStyle::Get().GetColor("Tutorials.Content.Color.Hovered") : FEditorStyle::Get().GetColor("Tutorials.Content.Color");
 }
 
 float STutorialContent::GetAnimatedZoom() const
 {
-	if(ContentIntroAnimation.IsPlaying())
+	if(ContentIntroAnimation.IsPlaying() && FSlateApplication::Get().IsRunningAtTargetFrameRate())
 	{
 		FIntroTutorials& IntroTutorials = FModuleManager::GetModuleChecked<FIntroTutorials>(TEXT("IntroTutorials"));
 		return 0.75f + (0.25f * IntroTutorials.GetIntroCurveValue(ContentIntroAnimation.GetLerp()));
@@ -380,4 +444,130 @@ float STutorialContent::GetInverseAnimatedZoom() const
 	return 1.0f / GetAnimatedZoom();
 }
 
+TSharedRef<SWidget> STutorialContent::HandleGetMenuContent()
+{
+	const bool bInShouldCloseWindowAfterMenuSelection = true;
+	FMenuBuilder MenuBuilder(bInShouldCloseWindowAfterMenuSelection, TSharedPtr<const FUICommandList>());
+
+	MenuBuilder.BeginSection(TEXT("Tutorial Options"));
+	{
+		MenuBuilder.AddMenuEntry(
+			LOCTEXT("ExitLabel", "Exit"),
+			LOCTEXT("ExitTooltip", "Quit this tutorial. You can find it again in the tutorials browser, reached from the Help menu."),
+			FSlateIcon(),
+			FUIAction(
+				FExecuteAction::CreateSP(this, &STutorialContent::HandleExitSelected),
+				FCanExecuteAction()
+			)
+		);
+
+		if(IsBackEnabled.Get())
+		{
+			MenuBuilder.AddMenuEntry(
+				LOCTEXT("BackLabel", "Back"),
+				LOCTEXT("BackTooltip", "Go back to the previous stage of this tutorial."),
+				FSlateIcon(),
+				FUIAction(
+					FExecuteAction::CreateSP(this, &STutorialContent::HandleBackSelected),
+					FCanExecuteAction()
+				)
+			);
+
+			MenuBuilder.AddMenuEntry(
+				LOCTEXT("RestartLabel", "Restart"),
+				LOCTEXT("RestartTooltip", "Start this tutorial again from the beginning."),
+				FSlateIcon(),
+				FUIAction(
+					FExecuteAction::CreateSP(this, &STutorialContent::HandleRestartSelected),
+					FCanExecuteAction()
+				)
+			);
+		}
+
+		if(IsHomeEnabled.Get())
+		{
+			MenuBuilder.AddMenuEntry(
+				LOCTEXT("OpenBrowserLabel", "More Tutorials..."),
+				LOCTEXT("OpenBrowserTooltip", "Open the tutorial browser to find more tutorials."),
+				FSlateIcon(),
+				FUIAction(
+					FExecuteAction::CreateSP(this, &STutorialContent::HandleBrowseSelected),
+					FCanExecuteAction()
+				)
+			);	
+		}
+	}
+	MenuBuilder.EndSection();
+
+	return MenuBuilder.MakeWidget();
+}
+
+void STutorialContent::HandleExitSelected()
+{
+	OnClosed.ExecuteIfBound();
+}
+
+void STutorialContent::HandleBackSelected()
+{
+	OnBackClicked.ExecuteIfBound();
+}
+
+void STutorialContent::HandleRestartSelected()
+{
+	if(Tutorial.IsValid())
+	{
+		FIntroTutorials& IntroTutorials = FModuleManager::GetModuleChecked<FIntroTutorials>(TEXT("IntroTutorials"));
+		const bool bRestart = true;
+		IntroTutorials.LaunchTutorial(Tutorial.Get(), bRestart, FSlateApplication::Get().FindWidgetWindow(AsShared()));		
+	}
+}
+
+void STutorialContent::HandleBrowseSelected()
+{
+	OnHomeClicked.ExecuteIfBound();
+}
+
+FReply STutorialContent::HandleNextClicked()
+{
+	if(IsNextEnabled.Get())
+	{
+		OnNextClicked.ExecuteIfBound();
+	}
+	else
+	{
+		OnHomeClicked.ExecuteIfBound();
+	}
+	
+	return FReply::Handled();
+}
+
+const FSlateBrush* STutorialContent::GetNextButtonBrush() const
+{
+	if(IsNextEnabled.Get())
+	{
+		return FEditorStyle::GetBrush("Tutorials.Navigation.NextButton");
+	}
+	else
+	{
+		return FEditorStyle::GetBrush("Tutorials.Navigation.HomeButton");
+	}
+}
+
+FText STutorialContent::GetNextButtonTooltip() const
+{
+	if(IsNextEnabled.Get())
+	{
+		return LOCTEXT("NextButtonTooltip", "Go to the next stage of this tutorial.");
+	}
+	else
+	{
+		return LOCTEXT("HomeButtonTooltip", "This tutorial is complete. Open the tutorial browser to find more tutorials.");
+	}
+}
+
+FSlateColor STutorialContent::GetNextButtonColor() const
+{
+	return NextButton->IsHovered() ? FLinearColor(0.1f, 0.1f, 0.1f, 1.0f) : FLinearColor(0.0f, 0.0f, 0.0f, 1.0f);
+}
+
 #undef LOCTEXT_NAMESPACE
\ No newline at end of file
diff --git a/Engine/Source/Editor/IntroTutorials/Private/STutorialContent.h b/Engine/Source/Editor/IntroTutorials/Private/STutorialContent.h
--- a/Engine/Source/Editor/IntroTutorials/Private/STutorialContent.h
+++ b/Engine/Source/Editor/IntroTutorials/Private/STutorialContent.h
@@ -32,6 +32,24 @@ class STutorialContent : public SCompoundWidget
 	/** Delegate fired when the close button is clicked */
 	SLATE_ARGUMENT(FSimpleDelegate, OnClosed)
 
+	/** Delegate fired when the back button is clicked */
+	SLATE_ARGUMENT(FSimpleDelegate, OnBackClicked)
+
+	/** Delegate fired when the home button is clicked */
+	SLATE_ARGUMENT(FSimpleDelegate, OnHomeClicked)
+
+	/** Delegate fired when the next button is clicked */
+	SLATE_ARGUMENT(FSimpleDelegate, OnNextClicked)
+
+	/** Attribute controlling enabled state of back functionality */
+	SLATE_ATTRIBUTE(bool, IsBackEnabled)
+
+	/** Attribute controlling enabled state of home functionality */
+	SLATE_ATTRIBUTE(bool, IsHomeEnabled)
+
+	/** Attribute controlling enabled state of next functionality */
+	SLATE_ATTRIBUTE(bool, IsNextEnabled)
+
 	/** Where text should be wrapped */
 	SLATE_ARGUMENT(float, WrapTextAt)
 
@@ -40,7 +58,7 @@ class STutorialContent : public SCompoundWidget
 
 	SLATE_END_ARGS()
 
-	void Construct(const FArguments& InArgs, const FTutorialContent& InContent);
+	void Construct(const FArguments& InArgs, UEditorTutorial* InTutorial, const FTutorialContent& InContent);
 
 	/** SWidget implementation */
 	virtual int32 OnPaint( const FPaintArgs& Args, const FGeometry& AllottedGeometry, const FSlateRect& MyClippingRect, FSlateWindowElementList& OutDrawElements, int32 LayerId, const FWidgetStyle& InWidgetStyle, bool bParentEnabled ) const override;
@@ -74,6 +92,9 @@ private:
 	/** Get close button visibility - varies depending on whether we are standalone or not */
 	EVisibility GetCloseButtonVisibility() const;
 
+	/** Get menu button visibility - varies depending on whether we are standalone or not */
+	EVisibility GetMenuButtonVisibility() const;
+
 	/** Alter the background color depending on hover state */
 	FSlateColor GetBackgroundColor() const;
 
@@ -83,6 +104,33 @@ private:
 	/** Get inverse zoom level padding for content - needed because rich text content doesnt scale well */
 	float GetInverseAnimatedZoom() const;
 
+	/** Get the content for the navigation menu */
+	TSharedRef<SWidget> HandleGetMenuContent();
+
+	/** Delegate handler for exiting the tutorial */
+	void HandleExitSelected();
+
+	/** Delegate handler for going to the previous stage of the tutorial */
+	void HandleBackSelected();
+
+	/** Delegate handler for restarting the tutorial */
+	void HandleRestartSelected();
+
+	/** Delegate handler for exiting the tutorial to the browser */
+	void HandleBrowseSelected();
+
+	/** Delegate handler for going to the next stage of the tutorial */
+	FReply HandleNextClicked();
+
+	/** Delegate handler allowing us to change the brush of the 'next' button depending on context */
+	const FSlateBrush* GetNextButtonBrush() const;
+
+	/** Delegate handler allowing us to change the tootlip of the 'next' button depending on context */
+	FText GetNextButtonTooltip() const;
+
+	/** Chage next button color based on hover state */
+	FSlateColor GetNextButtonColor() const;
+
 private:
 
 	/** Copy of the window size we were last draw at */
@@ -118,6 +166,24 @@ private:
 	/** Delegate fired when the close button is clicked */
 	FSimpleDelegate OnClosed;
 
+	/** Delegate fired when the next button is clicked */
+	FSimpleDelegate OnNextClicked;
+
+	/** Delegate fired when the home button is clicked */
+	FSimpleDelegate OnHomeClicked;
+
+	/** Delegate fired when the back button is clicked */
+	FSimpleDelegate OnBackClicked;
+
+	/** Attribute controlling enabled state of back functionality */
+	TAttribute<bool> IsBackEnabled;
+
+	/** Attribute controlling enabled state of home functionality */
+	TAttribute<bool> IsHomeEnabled;
+
+	/** Attribute controlling enabled state of next functionality */
+	TAttribute<bool> IsNextEnabled;
+
 	/** Animation curves for displaying border */
 	FCurveSequence BorderPulseAnimation;
 	FCurveSequence BorderIntroAnimation;
@@ -127,4 +193,10 @@ private:
 
 	/** Documentation page reference to use if we are displaying a UDN doc - we need this otherwise the page will be freed */
 	TSharedPtr<IDocumentationPage> DocumentationPage;
+
+	/** The tutorial we are referencing */
+	TWeakObjectPtr<UEditorTutorial> Tutorial;
+
+	/** Next button widget */
+	TSharedPtr<SWidget> NextButton;
 };
\ No newline at end of file
diff --git a/Engine/Source/Editor/IntroTutorials/Private/STutorialOverlay.cpp b/Engine/Source/Editor/IntroTutorials/Private/STutorialOverlay.cpp
--- a/Engine/Source/Editor/IntroTutorials/Private/STutorialOverlay.cpp
+++ b/Engine/Source/Editor/IntroTutorials/Private/STutorialOverlay.cpp
@@ -10,7 +10,7 @@
 #include "BlueprintEditorUtils.h"
 #include "Guid.h"
 
-void STutorialOverlay::Construct(const FArguments& InArgs, FTutorialStage* const InStage)
+void STutorialOverlay::Construct(const FArguments& InArgs, UEditorTutorial* InTutorial, FTutorialStage* const InStage)
 {
 	ParentWindow = InArgs._ParentWindow;
 	bIsStandalone = InArgs._IsStandalone;
@@ -39,8 +39,14 @@ void STutorialOverlay::Construct(const FArguments& InArgs, FTutorialStage* const
 				.VAlign(VAlign_Center)
 				.HAlign(HAlign_Center)
 				[
-					SNew(STutorialContent, InStage->Content)
+					SNew(STutorialContent, InTutorial, InStage->Content)
 					.OnClosed(InArgs._OnClosed)
+					.OnNextClicked(InArgs._OnNextClicked)
+					.OnHomeClicked(InArgs._OnHomeClicked)
+					.OnBackClicked(InArgs._OnBackClicked)
+					.IsBackEnabled(InArgs._IsBackEnabled)
+					.IsHomeEnabled(InArgs._IsHomeEnabled)
+					.IsNextEnabled(InArgs._IsNextEnabled)
 					.IsStandalone(InArgs._IsStandalone)
 					.WrapTextAt(600.0f)
 				]
@@ -56,14 +62,19 @@ void STutorialOverlay::Construct(const FArguments& InArgs, FTutorialStage* const
 			{
 				if (WidgetContent.Content.Type != ETutorialContent::None)
 				{
-
 					TSharedPtr<STutorialContent> ContentWidget =
-						SNew(STutorialContent, WidgetContent.Content)
+						SNew(STutorialContent, InTutorial, WidgetContent.Content)
 						.HAlign(WidgetContent.HorizontalAlignment)
 						.VAlign(WidgetContent.VerticalAlignment)
 						.Offset(WidgetContent.Offset)
 						.IsStandalone(bIsStandalone)
 						.OnClosed(OnClosed)
+						.OnNextClicked(InArgs._OnNextClicked)
+						.OnHomeClicked(InArgs._OnHomeClicked)
+						.OnBackClicked(InArgs._OnBackClicked)
+						.IsBackEnabled(InArgs._IsBackEnabled)
+						.IsHomeEnabled(InArgs._IsHomeEnabled)
+						.IsNextEnabled(InArgs._IsNextEnabled)
 						.WrapTextAt(WidgetContent.ContentWidth)
 						.Anchor(WidgetContent.WidgetAnchor);
 					
diff --git a/Engine/Source/Editor/IntroTutorials/Private/STutorialOverlay.h b/Engine/Source/Editor/IntroTutorials/Private/STutorialOverlay.h
--- a/Engine/Source/Editor/IntroTutorials/Private/STutorialOverlay.h
+++ b/Engine/Source/Editor/IntroTutorials/Private/STutorialOverlay.h
@@ -32,12 +32,19 @@ public:
 	/** Delegate fired when the close button is clicked */
 	SLATE_ARGUMENT(FSimpleDelegate, OnClosed)
 
+	SLATE_ARGUMENT(FSimpleDelegate, OnBackClicked)
+	SLATE_ARGUMENT(FSimpleDelegate, OnHomeClicked)
+	SLATE_ARGUMENT(FSimpleDelegate, OnNextClicked)
+	SLATE_ATTRIBUTE(bool, IsBackEnabled)
+	SLATE_ATTRIBUTE(bool, IsHomeEnabled)
+	SLATE_ATTRIBUTE(bool, IsNextEnabled)
+
 	/** Whether we can show full window content in this overlay (i.e. in the same window as the navigation controls) */
 	SLATE_ARGUMENT(bool, AllowNonWidgetContent)
 
 	SLATE_END_ARGS()
 
-	void Construct(const FArguments& InArgs, FTutorialStage* const InStage);
+	void Construct(const FArguments& InArgs, UEditorTutorial* InTutorial, FTutorialStage* const InStage);
 
 	/** SWidget implementation */
 	virtual int32 OnPaint( const FPaintArgs& Args, const FGeometry& AllottedGeometry, const FSlateRect& MyClippingRect, FSlateWindowElementList& OutDrawElements, int32 LayerId, const FWidgetStyle& InWidgetStyle, bool bParentEnabled ) const override;
