Tutorials can now open an asset when started

The tutorial system will then attach non-widget-bound content to the asset editor window.
Also, assets can now be opened from a tutorial blueprint.

diff --git a/Engine/Source/Editor/IntroTutorials/Private/EditorTutorial.cpp b/Engine/Source/Editor/IntroTutorials/Private/EditorTutorial.cpp
--- a/Engine/Source/Editor/IntroTutorials/Private/EditorTutorial.cpp
+++ b/Engine/Source/Editor/IntroTutorials/Private/EditorTutorial.cpp
@@ -44,3 +44,8 @@ void UEditorTutorial::HandleTutorialStageEnded(FName StageName)
 }
 
 
+void UEditorTutorial::OpenAsset(UObject* Asset)
+{
+	FAssetEditorManager::Get().OpenEditorForAsset(Asset);
+}
+
diff --git a/Engine/Source/Editor/IntroTutorials/Private/STutorialRoot.cpp b/Engine/Source/Editor/IntroTutorials/Private/STutorialRoot.cpp
--- a/Engine/Source/Editor/IntroTutorials/Private/STutorialRoot.cpp
+++ b/Engine/Source/Editor/IntroTutorials/Private/STutorialRoot.cpp
@@ -5,6 +5,8 @@
 #include "SEditorTutorials.h"
 #include "EditorTutorialSettings.h"
 #include "TutorialStateSettings.h"
+#include "AssetEditorManager.h"
+#include "ToolkitManager.h"
 
 #define LOCTEXT_NAMESPACE "STutorialRoot"
 
@@ -86,6 +88,24 @@ void STutorialRoot::LaunchTutorial(UEditorTutorial* InTutorial, bool bInRestart,
 		bool bHaveSeenTutorial = false;
 		CurrentTutorialStage = bInRestart ? 0 : GetDefault<UTutorialStateSettings>()->GetProgress(CurrentTutorial, bHaveSeenTutorial);
 
+		// check if we should be launching this tutorial for an asset editor
+		if(InTutorial->AssetToUse.IsValid())
+		{
+			TArray<FString> AssetPaths;
+			AssetPaths.Add(InTutorial->AssetToUse.AssetLongPathname);
+			FAssetEditorManager::Get().OpenEditorsForAssets(AssetPaths);
+
+			UObject* Asset = InTutorial->AssetToUse.ResolveObject();
+			if(Asset != nullptr)
+			{
+				TSharedPtr<IToolkit> Toolkit = FToolkitManager::Get().FindEditorForAsset( Asset );
+				if(Toolkit.IsValid())
+				{
+					InNavigationWindow = FSlateApplication::Get().FindWidgetWindow(Toolkit->GetToolkitHost()->GetParentWidget());
+				}
+			}
+		}
+
 		// launch tutorial for all windows we wrap - any tutorial can display over any window
 		for(auto& TutorialWidget : TutorialWidgets)
 		{
diff --git a/Engine/Source/Editor/IntroTutorials/Public/EditorTutorial.h b/Engine/Source/Editor/IntroTutorials/Public/EditorTutorial.h
--- a/Engine/Source/Editor/IntroTutorials/Public/EditorTutorial.h
+++ b/Engine/Source/Editor/IntroTutorials/Public/EditorTutorial.h
@@ -236,6 +236,10 @@ class INTROTUTORIALS_API UEditorTutorial : public UObject
 	UPROPERTY(EditAnywhere, Category="Tutorial")
 	bool bIsStandalone;
 	
+	/** Asset to open & attach the tutorial to. Non-widget-bound content will appear in the asset's window */
+	UPROPERTY(EditAnywhere, Category="Tutorial")
+	FStringAssetReference AssetToUse;
+
 public:
 	/** Called when a tutorial stage is started */
 	void HandleTutorialStageStarted(FName StageName);
@@ -266,4 +270,11 @@ protected:
 	/** Begin a tutorial. Note that this will end the current tutorial that is in progress, if any */
 	UFUNCTION(BlueprintCallable, Category="Tutorial")
 	static void BeginTutorial(UEditorTutorial* TutorialToStart, bool bRestart);
+
+	/** 
+	 * Open an asset for use by a tutorial
+	 * @param	Asset	The asset to open
+	 */
+	UFUNCTION(BlueprintCallable, Category="Tutorial")
+	static void OpenAsset(UObject* Asset);
 };
\ No newline at end of file
