Tutorial factories now import in the standard way

Added new factory for import.
Removed old import dialog - importing is now done via the import button or reimport... in the asset editor.
Added path that the file was imported from to UEditorTutorial.

diff --git a/Engine/Source/Editor/IntroTutorials/Private/EditorTutorialFactory.cpp b/Engine/Source/Editor/IntroTutorials/Private/EditorTutorialFactory.cpp
--- a/Engine/Source/Editor/IntroTutorials/Private/EditorTutorialFactory.cpp
+++ b/Engine/Source/Editor/IntroTutorials/Private/EditorTutorialFactory.cpp
@@ -4,7 +4,6 @@
 #include "EditorTutorialFactory.h"
 #include "EditorTutorial.h"
 #include "KismetEditorUtilities.h"
-#include "SEditorTutorialImporter.h"
 #include "IAssetTypeActions.h"
 
 #define LOCTEXT_NAMESPACE "UEditorTutorialFactory"
@@ -17,37 +16,9 @@ UEditorTutorialFactory::UEditorTutorialFactory(const class FPostConstructInitial
 	SupportedClass = UBlueprint::StaticClass();
 }
 
-bool UEditorTutorialFactory::ConfigureProperties()
-{
-	// show import dialog
-	TSharedRef<SWindow> Window = SNew(SWindow)
-			.Title(LOCTEXT("WindowTitle", "Import Tutorial"))
-			.SupportsMaximize(false)
-			.SupportsMinimize(false)
-			.SizingRule(ESizingRule::Autosized);	
-
-	ImporterDialog = SNew(SEditorTutorialImporter)
-		.ParentWindow(Window);
-
-	Window->SetContent(ImporterDialog.ToSharedRef());
-
-	TSharedPtr<SWindow> RootWindow = FGlobalTabmanager::Get()->GetRootWindow();
-	FSlateApplication::Get().AddModalWindow(Window, RootWindow);
-
-	return true;
-}
-
 UObject* UEditorTutorialFactory::FactoryCreateNew(UClass* Class, UObject* InParent, FName Name, EObjectFlags Flags, UObject* Context, FFeedbackContext* Warn)
 {
-	UBlueprint* NewBlueprint = FKismetEditorUtilities::CreateBlueprint(UEditorTutorial::StaticClass(), InParent, Name, BPTYPE_Normal, UBlueprint::StaticClass(), UBlueprintGeneratedClass::StaticClass(), NAME_None);
-	
-	if(ImporterDialog.IsValid() && !ImporterDialog->GetImportPath().IsEmpty())
-	{
-		ImporterDialog->Import(CastChecked<UEditorTutorial>(NewBlueprint->GeneratedClass->GetDefaultObject()));
-		ImporterDialog = nullptr;
-	}
-
-	return NewBlueprint;
+	return FKismetEditorUtilities::CreateBlueprint(UEditorTutorial::StaticClass(), InParent, Name, BPTYPE_Normal, UBlueprint::StaticClass(), UBlueprintGeneratedClass::StaticClass(), NAME_None);
 }
 
 uint32 UEditorTutorialFactory::GetMenuCategories() const
diff --git a/Engine/Source/Editor/IntroTutorials/Private/EditorTutorialFactory.h b/Engine/Source/Editor/IntroTutorials/Private/EditorTutorialFactory.h
--- a/Engine/Source/Editor/IntroTutorials/Private/EditorTutorialFactory.h
+++ b/Engine/Source/Editor/IntroTutorials/Private/EditorTutorialFactory.h
@@ -4,21 +4,14 @@
 
 #include "EditorTutorialFactory.generated.h"
 
-class SEditorTutorialImporter;
-
 UCLASS(hidecategories=Object, collapsecategories)
 class UEditorTutorialFactory : public UFactory
 {
 	GENERATED_UCLASS_BODY()
 
 	// UFactory Interface
-	virtual bool ConfigureProperties() override;
 	virtual UObject* FactoryCreateNew(UClass* Class, UObject* InParent, FName Name, EObjectFlags Flags, UObject* Context, FFeedbackContext* Warn) override;
 	virtual uint32 GetMenuCategories() const;
 	virtual FText GetDisplayName() const;
 	// End of UFactory Interface
-
-private:
-	/** Importer dialog reference */
-	TSharedPtr<SEditorTutorialImporter> ImporterDialog;
 };
diff --git a/Engine/Source/Editor/IntroTutorials/Private/EditorTutorialImportFactory.cpp b/Engine/Source/Editor/IntroTutorials/Private/EditorTutorialImportFactory.cpp
new file mode 100644
--- /dev/null
+++ b/Engine/Source/Editor/IntroTutorials/Private/EditorTutorialImportFactory.cpp
@@ -0,0 +1,168 @@
+// Copyright 1998-2014 Epic Games, Inc. All Rights Reserved.
+
+#include "IntroTutorialsPrivatePCH.h"
+#include "EditorTutorialImportFactory.h"
+#include "EditorTutorial.h"
+#include "KismetEditorUtilities.h"
+#include "IAssetTypeActions.h"
+#include "IDocumentation.h"
+#include "IDocumentationPage.h"
+
+#define LOCTEXT_NAMESPACE "UEditorTutorialImportFactory"
+
+UEditorTutorialImportFactory::UEditorTutorialImportFactory(const class FPostConstructInitializeProperties& PCIP)
+	: Super(PCIP)
+{
+	bEditorImport = true;
+	bEditAfterNew = true;
+	bText = false;
+	SupportedClass = UBlueprint::StaticClass();
+	Formats.Add(TEXT("udn;UDN documentation files"));
+}
+
+bool UEditorTutorialImportFactory::FactoryCanImport(const FString& Filename)
+{
+	const FString DocDir = FPaths::ConvertRelativePathToFull(FPaths::EngineDir() / TEXT("Documentation/Source"));
+	FString NewPath = FPaths::ConvertRelativePathToFull(FPaths::GetPath(Filename));
+	FPaths::NormalizeFilename(NewPath);
+
+	return (NewPath.StartsWith(DocDir));
+}
+
+static bool GetDocumentationPath(FString& InOutFilePath)
+{
+	const FString DocDir = FPaths::ConvertRelativePathToFull(FPaths::EngineDir() / TEXT("Documentation/Source"));
+	FString NewPath = FPaths::ConvertRelativePathToFull(FPaths::GetPath(InOutFilePath));
+	FPaths::NormalizeFilename(NewPath);
+
+	if(NewPath.StartsWith(DocDir))
+	{
+		InOutFilePath = NewPath.RightChop(DocDir.Len());
+		return true;
+	}
+
+	return false;
+}
+
+UObject* UEditorTutorialImportFactory::FactoryCreateBinary(UClass* InClass, UObject* InParent, FName InName, EObjectFlags Flags, UObject* Context, const TCHAR* Type, const uint8*& Buffer, const uint8* BufferEnd, FFeedbackContext* Warn)
+{
+	UBlueprint* NewBlueprint = FKismetEditorUtilities::CreateBlueprint(UEditorTutorial::StaticClass(), InParent, InName, BPTYPE_Normal, UBlueprint::StaticClass(), UBlueprintGeneratedClass::StaticClass(), NAME_None);
+	
+	FString PathToImport = GetCurrentFilename();
+	if(GetDocumentationPath(PathToImport))
+	{
+		UEditorTutorial* EditorTutorial = CastChecked<UEditorTutorial>(NewBlueprint->GeneratedClass->GetDefaultObject());
+		if(Import(EditorTutorial, PathToImport))
+		{
+			EditorTutorial->ImportPath = GetCurrentFilename();
+		}
+	}
+
+	return NewBlueprint;
+}
+
+bool UEditorTutorialImportFactory::CanReimport(UObject* Obj, TArray<FString>& OutFilenames)
+{
+	UBlueprint* Blueprint = Cast<UBlueprint>(Obj);
+	if(Blueprint != nullptr)
+	{
+		UEditorTutorial* EditorTutorial = Cast<UEditorTutorial>(Blueprint->GeneratedClass->GetDefaultObject());
+		if(EditorTutorial != nullptr)
+		{
+			OutFilenames.Add(EditorTutorial->ImportPath);
+			return true;
+		}
+	}
+
+	return false;
+}
+
+void UEditorTutorialImportFactory::SetReimportPaths(UObject* Obj, const TArray<FString>& NewReimportPaths)
+{
+	if(NewReimportPaths.Num() > 0)
+	{
+		UBlueprint* Blueprint = Cast<UBlueprint>(Obj);
+		if(Blueprint != nullptr)
+		{
+			UEditorTutorial* EditorTutorial = Cast<UEditorTutorial>(Blueprint->GeneratedClass->GetDefaultObject());
+			if(EditorTutorial != nullptr)
+			{
+				if(FactoryCanImport(NewReimportPaths[0]))
+				{
+					EditorTutorial->ImportPath = NewReimportPaths[0];
+				}
+			}
+		}
+	}
+}
+
+EReimportResult::Type UEditorTutorialImportFactory::Reimport(UObject* Obj)
+{
+	UBlueprint* Blueprint = Cast<UBlueprint>(Obj);
+	if(Blueprint != nullptr)
+	{
+		UEditorTutorial* EditorTutorial = Cast<UEditorTutorial>(Blueprint->GeneratedClass->GetDefaultObject());
+		if(EditorTutorial != nullptr)
+		{
+			FString PathToImport = EditorTutorial->ImportPath;
+			if(GetDocumentationPath(PathToImport))
+			{
+				return Import(EditorTutorial, PathToImport) ? EReimportResult::Succeeded : EReimportResult::Failed;
+			}
+		}
+	}
+
+	return EReimportResult::Failed;
+}
+
+bool UEditorTutorialImportFactory::Import(UEditorTutorial* InTutorialToImportTo, const FString& InImportPath)
+{
+	if(!InImportPath.IsEmpty() && InTutorialToImportTo != nullptr && IDocumentation::Get()->PageExists(InImportPath))
+	{
+		FDocumentationStyle DocumentationStyle;
+		DocumentationStyle
+			.ContentStyle(TEXT("Tutorials.Content.Text"))
+			.BoldContentStyle(TEXT("Tutorials.Content.TextBold"))
+			.NumberedContentStyle(TEXT("Tutorials.Content.Text"))
+			.Header1Style(TEXT("Tutorials.Content.HeaderText1"))
+			.Header2Style(TEXT("Tutorials.Content.HeaderText2"))
+			.HyperlinkStyle(TEXT("Tutorials.Content.Hyperlink"))
+			.HyperlinkTextStyle(TEXT("Tutorials.Content.HyperlinkText"))
+			.SeparatorStyle(TEXT("Tutorials.Separator"));
+
+		TSharedRef<IDocumentationPage> Page = IDocumentation::Get()->GetPage(InImportPath, TSharedPtr<FParserConfiguration>(), DocumentationStyle);
+		InTutorialToImportTo->Modify();
+		InTutorialToImportTo->Title = Page->GetTitle();
+		InTutorialToImportTo->Stages.Empty();
+
+		TArray<FExcerpt> Excerpts;
+		Page->GetExcerpts(Excerpts);
+		for(auto& Excerpt : Excerpts)
+		{
+			Page->GetExcerptContent(Excerpt);
+
+			if(Excerpt.RichText.Len() > 0)
+			{
+				FTutorialStage& Stage = InTutorialToImportTo->Stages[InTutorialToImportTo->Stages.Add(FTutorialStage())];
+				Stage.Name = *Excerpt.Name;
+				Stage.Content.Type = ETutorialContent::RichText;
+
+				FString RichText;
+				FString* TitleStringPtr = Excerpt.Variables.Find(TEXT("StageTitle"));
+				if(TitleStringPtr != nullptr)
+				{
+					RichText += FString::Printf(TEXT("<TextStyle Style=\"Tutorials.Content.HeaderText2\">%s</>\n\n"), **TitleStringPtr);
+				}
+				RichText += Excerpt.RichText;
+
+				Stage.Content.Text = FText::FromString(RichText);
+			}
+		}
+
+		return true;
+	}
+
+	return false;
+}
+
+#undef LOCTEXT_NAMESPACE
\ No newline at end of file
diff --git a/Engine/Source/Editor/IntroTutorials/Private/EditorTutorialImportFactory.h b/Engine/Source/Editor/IntroTutorials/Private/EditorTutorialImportFactory.h
new file mode 100644
--- /dev/null
+++ b/Engine/Source/Editor/IntroTutorials/Private/EditorTutorialImportFactory.h
@@ -0,0 +1,26 @@
+// Copyright 1998-2014 Epic Games, Inc. All Rights Reserved.
+
+#pragma once
+
+#include "EditorTutorialImportFactory.generated.h"
+
+UCLASS(hidecategories=Object, collapsecategories)
+class UEditorTutorialImportFactory : public UFactory, public FReimportHandler
+{
+	GENERATED_UCLASS_BODY()
+
+	// UFactory Interface
+	virtual bool FactoryCanImport(const FString& Filename) override;
+	virtual UObject* FactoryCreateBinary(UClass* InClass, UObject* InParent, FName InName, EObjectFlags Flags, UObject* Context, const TCHAR* Type, const uint8*& Buffer, const uint8* BufferEnd, FFeedbackContext* Warn) override;
+	// End of UFactory Interface
+
+	// Begin FReimportHandler interface
+	virtual bool CanReimport(UObject* Obj, TArray<FString>& OutFilenames) override;
+	virtual void SetReimportPaths(UObject* Obj, const TArray<FString>& NewReimportPaths) override;
+	virtual EReimportResult::Type Reimport(UObject* Obj) override;
+	// End FReimportHandler interface
+
+private:
+	/** Import the tutorial from file */
+	bool Import(UEditorTutorial* InTutorialToImportTo, const FString& InImportPath);
+};
diff --git a/Engine/Source/Editor/IntroTutorials/Private/SEditorTutorialImporter.cpp b/Engine/Source/Editor/IntroTutorials/Private/SEditorTutorialImporter.cpp
deleted file mode 100644
diff --git a/Engine/Source/Editor/IntroTutorials/Private/SEditorTutorialImporter.h b/Engine/Source/Editor/IntroTutorials/Private/SEditorTutorialImporter.h
deleted file mode 100644
diff --git a/Engine/Source/Editor/IntroTutorials/Public/EditorTutorial.h b/Engine/Source/Editor/IntroTutorials/Public/EditorTutorial.h
--- a/Engine/Source/Editor/IntroTutorials/Public/EditorTutorial.h
+++ b/Engine/Source/Editor/IntroTutorials/Public/EditorTutorial.h
@@ -240,6 +240,10 @@ class INTROTUTORIALS_API UEditorTutorial : public UObject
 	UPROPERTY(EditAnywhere, Category="Tutorial")
 	FStringAssetReference AssetToUse;
 
+	/** The path this tutorial was imported from, if any. */
+	UPROPERTY()
+	FString ImportPath;
+
 public:
 	/** Called when a tutorial stage is started */
 	void HandleTutorialStageStarted(FName StageName);
