CIS fix - non-unity build

diff --git a/Engine/Source/Editor/IntroTutorials/Private/STutorialRoot.cpp b/Engine/Source/Editor/IntroTutorials/Private/STutorialRoot.cpp
--- a/Engine/Source/Editor/IntroTutorials/Private/STutorialRoot.cpp
+++ b/Engine/Source/Editor/IntroTutorials/Private/STutorialRoot.cpp
@@ -7,6 +7,8 @@
 #include "TutorialStateSettings.h"
 #include "AssetEditorManager.h"
 #include "ToolkitManager.h"
+#include "IToolkit.h"
+#include "IToolkitHost.h"
 
 #define LOCTEXT_NAMESPACE "STutorialRoot"
 
