Added/edited extra tutorial contexts & added analytics

Tutorial contexts for IOS/Android/Whatever are now hooked up (console setup ones are still missing).

Analytics added for tutorial usage in various places.

Added ability to reset tutorial state with -ResetTutorials command-line flag.

Cleaned up some unused code (still a lot more to come here!).

diff --git a/Engine/Plugins/2D/Paper2D/Source/Paper2DEditor/Private/Paper2DEditorModule.cpp b/Engine/Plugins/2D/Paper2D/Source/Paper2DEditor/Private/Paper2DEditorModule.cpp
--- a/Engine/Plugins/2D/Paper2D/Source/Paper2DEditor/Private/Paper2DEditorModule.cpp
+++ b/Engine/Plugins/2D/Paper2D/Source/Paper2DEditor/Private/Paper2DEditorModule.cpp
@@ -137,7 +137,6 @@ public:
 		FPaperContentBrowserExtensions::InstallHooks();
 
 		RegisterSettings();
-		RegisterIntroTutorials();
 	}
 
 	virtual void ShutdownModule() override
@@ -150,7 +149,6 @@ public:
 
 		if (UObjectInitialized())
 		{
-			UnregisterIntroTutorials();
 			UnregisterSettings();
 
 			FPaperContentBrowserExtensions::RemoveHooks();
@@ -232,34 +230,6 @@ private:
 		}
 	}
 
-	void RegisterIntroTutorials()
-	{
-		if (!IsRunningCommandlet())
-		{
-			//@TODO: PAPER2D: Remove the _Preview suffix on the config keys once the final doc is in place
-			// (this is so that people who dismiss the early warning message still get the final intro doc later on)
-			IIntroTutorials::Get().RegisterTutorialForAssetEditor(
-				UPaperSprite::StaticClass(),
-				TEXT("Shared/Tutorials/InPaperSpriteEditorTutorial"),
-				TEXT("SeenPaperSpriteEditorWelcome_Preview"),
-				FString("19D7EA18-629B-4A86-BD19-ED2B3BE53600"));
-			IIntroTutorials::Get().RegisterTutorialForAssetEditor(
-				UPaperFlipbook::StaticClass(),
-				TEXT("Shared/Tutorials/InPaperFlipbookEditorTutorial"),
-				TEXT("SeenPaperFlipbookEditorWelcome_Preview"),
-				FString("B24214C1-E17A-4F95-BE18-2ED8BFCEC008"));
-		}
-	}
-
-	void UnregisterIntroTutorials()
-	{
-		if (IIntroTutorials::IsAvailable())
-		{
-			IIntroTutorials::Get().UnregisterTutorialForAssetEditor(UPaperSprite::StaticClass());
-			IIntroTutorials::Get().UnregisterTutorialForAssetEditor(UPaperFlipbook::StaticClass());
-		}
-	}
-
 	// Installs or uninstalls the tile map editing mode depending on settings
 	void UpdateTileMapEditorModeInstallation()
 	{
diff --git a/Engine/Source/Developer/DesktopPlatform/Private/PlatformInfo.cpp b/Engine/Source/Developer/DesktopPlatform/Private/PlatformInfo.cpp
--- a/Engine/Source/Developer/DesktopPlatform/Private/PlatformInfo.cpp
+++ b/Engine/Source/Developer/DesktopPlatform/Private/PlatformInfo.cpp
@@ -60,53 +60,53 @@ FPlatformInfo BuildPlatformInfo(const FName& InPlatformInfoName, const FName& In
 static const bool IsAvailableOnWindows = true;
 static const bool IsAvailableOnMac = false;
 static const bool IsAvailableOnLinux = false;
-static const FString IOSTutorial = TEXT("Shared/Tutorials/InstallingiTunesTutorial");
+static const FString IOSTutorial = TEXT("/Engine/Tutorial/Mobile/InstallingiTunesTutorial.InstallingiTunesTutorial");
 #elif PLATFORM_MAC
 static const bool IsAvailableOnWindows = false;
 static const bool IsAvailableOnMac = true;
 static const bool IsAvailableOnLinux = false;
-static const FString IOSTutorial = TEXT("Shared/Tutorials/InstallingXCodeTutorial");
+static const FString IOSTutorial = TEXT("/Engine/Tutorial/Installation/InstallingXCodeTutorial.InstallingXCodeTutorial");
 #elif PLATFORM_LINUX
 static const bool IsAvailableOnWindows = false;
 static const bool IsAvailableOnMac = false;
 static const bool IsAvailableOnLinux = true;
-static const FString IOSTutorial = TEXT("Shared/Tutorials/NotYetImplemented");
+static const FString IOSTutorial = TEXT("/Engine/Tutorial/NotYetImplemented");
 #endif
 
 static const FPlatformInfo PlatformInfoArray[] = {
-	//				  PlatformInfoName					TargetPlatformName			DisplayName													PlatformType				PlatformFlags					IconPaths																																			UATCommandLine											AutoSDKPath			SDKStatus						SDKTutorial													EnabledForUse								EnabledForBinary
-	BuildPlatformInfo(TEXT("WindowsNoEditor"),			TEXT("WindowsNoEditor"),	LOCTEXT("WindowsNoEditor", "Windows"),						EPlatformType::Game,		EPlatformFlags::None,			FPlatformIconPaths(TEXT("Launcher/WindowsTarget/Platform_WindowsNoEditor_24x"), TEXT("Launcher/WindowsTarget/Platform_WindowsNoEditor_128x")),		TEXT(""),												TEXT(""),			EPlatformSDKStatus::Unknown,	TEXT("Shared/Tutorials/InstallingVisualStudioTutorial"),	IsAvailableOnWindows,						true),
-	BuildPlatformInfo(TEXT("WindowsNoEditor_Win32"),	TEXT("WindowsNoEditor"),	LOCTEXT("WindowsNoEditor_Win32", "Windows (32-bit)"),		EPlatformType::Game,		EPlatformFlags::BuildFlavor,	FPlatformIconPaths(TEXT("Launcher/WindowsTarget/Platform_WindowsNoEditor_24x"), TEXT("Launcher/WindowsTarget/Platform_WindowsNoEditor_128x")),		TEXT("-targetplatform=Win32"),							TEXT(""),			EPlatformSDKStatus::Unknown,	TEXT("Shared/Tutorials/InstallingVisualStudioTutorial"),	IsAvailableOnWindows,						true),
-	BuildPlatformInfo(TEXT("WindowsNoEditor_Win64"),	TEXT("WindowsNoEditor"),	LOCTEXT("WindowsNoEditor_Win64", "Windows (64-bit)"),		EPlatformType::Game,		EPlatformFlags::BuildFlavor,	FPlatformIconPaths(TEXT("Launcher/WindowsTarget/Platform_WindowsNoEditor_24x"), TEXT("Launcher/WindowsTarget/Platform_WindowsNoEditor_128x")),		TEXT("-targetplatform=Win64"),							TEXT(""),			EPlatformSDKStatus::Unknown,	TEXT("Shared/Tutorials/InstallingVisualStudioTutorial"),	IsAvailableOnWindows,						true),
-	BuildPlatformInfo(TEXT("Windows"),					TEXT("Windows"),			LOCTEXT("WindowsEditor", "Windows (Editor)"),				EPlatformType::Editor,		EPlatformFlags::None,			FPlatformIconPaths(TEXT("Launcher/WindowsTarget/Platform_Windows_24x"), TEXT("Launcher/WindowsTarget/Platform_Windows_128x")),						TEXT(""),												TEXT(""),			EPlatformSDKStatus::Unknown,	TEXT("Shared/Tutorials/InstallingVisualStudioTutorial"),	IsAvailableOnWindows,						true),
-	BuildPlatformInfo(TEXT("WindowsClient"),			TEXT("WindowsClient"),		LOCTEXT("WindowsClient", "Windows (Client-only)"),			EPlatformType::Client,		EPlatformFlags::None,			FPlatformIconPaths(TEXT("Launcher/WindowsTarget/Platform_Windows_24x"), TEXT("Launcher/WindowsTarget/Platform_Windows_128x")),						TEXT(""),												TEXT(""),			EPlatformSDKStatus::Unknown,	TEXT("Shared/Tutorials/InstallingVisualStudioTutorial"),	IsAvailableOnWindows,						true),
-	BuildPlatformInfo(TEXT("WindowsServer"),			TEXT("WindowsServer"),		LOCTEXT("WindowsServer", "Windows (Dedicated Server)"),		EPlatformType::Server,		EPlatformFlags::None,			FPlatformIconPaths(TEXT("Launcher/WindowsTarget/Platform_WindowsServer_24x"), TEXT("Launcher/WindowsTarget/Platform_WindowsServer_128x")),			TEXT(""),												TEXT(""),			EPlatformSDKStatus::Unknown,	TEXT("Shared/Tutorials/InstallingVisualStudioTutorial"),	IsAvailableOnWindows,						true),
-
-	BuildPlatformInfo(TEXT("MacNoEditor"),				TEXT("MacNoEditor"),		LOCTEXT("MacNoEditor", "Mac"),								EPlatformType::Game,		EPlatformFlags::None,			FPlatformIconPaths(TEXT("Launcher/MacTarget/Platform_Mac_24x"), TEXT("Launcher/MacTarget/Platform_Mac_128x")),										TEXT(""),												TEXT(""),			EPlatformSDKStatus::Unknown,	TEXT("Shared/Tutorials/InstallingXCodeTutorial"),			IsAvailableOnMac,							true),
-	BuildPlatformInfo(TEXT("Mac"),						TEXT("Mac"),				LOCTEXT("MacEditor", "Mac (Editor)"),						EPlatformType::Editor,		EPlatformFlags::None,			FPlatformIconPaths(TEXT("Launcher/MacTarget/Platform_Mac_24x"), TEXT("Launcher/MacTarget/Platform_Mac_128x")),										TEXT(""),												TEXT(""),			EPlatformSDKStatus::Unknown,	TEXT("Shared/Tutorials/InstallingXCodeTutorial"),			IsAvailableOnMac,							true),
-	BuildPlatformInfo(TEXT("MacClient"),				TEXT("MacClient"),			LOCTEXT("MacClient", "Mac (Client-only)"),					EPlatformType::Client,		EPlatformFlags::None,			FPlatformIconPaths(TEXT("Launcher/MacTarget/Platform_Mac_24x"), TEXT("Launcher/MacTarget/Platform_Mac_128x")),										TEXT(""),												TEXT(""),			EPlatformSDKStatus::Unknown,	TEXT("Shared/Tutorials/InstallingXCodeTutorial"),			IsAvailableOnMac,							true),
-	BuildPlatformInfo(TEXT("MacServer"),				TEXT("MacServer"),			LOCTEXT("MacServer", "Mac (Dedicated Server)"),				EPlatformType::Server,		EPlatformFlags::None,			FPlatformIconPaths(TEXT("Launcher/MacTarget/Platform_Mac_24x"), TEXT("Launcher/MacTarget/Platform_Mac_128x")),										TEXT(""),												TEXT(""),			EPlatformSDKStatus::Unknown,	TEXT("Shared/Tutorials/InstallingXCodeTutorial"),			IsAvailableOnMac,							true),
-
-	BuildPlatformInfo(TEXT("LinuxNoEditor"),			TEXT("LinuxNoEditor"),		LOCTEXT("LinuxNoEditor", "Linux"),							EPlatformType::Game,		EPlatformFlags::None,			FPlatformIconPaths(TEXT("Launcher/LinuxTarget/Platform_Linux_24x"), TEXT("Launcher/LinuxTarget/Platform_Linux_128x")),								TEXT(""),												TEXT("Linux_x64"),	EPlatformSDKStatus::Unknown,	TEXT(""),													IsAvailableOnLinux || IsAvailableOnWindows,	false),	
-	BuildPlatformInfo(TEXT("Linux"),					TEXT("Linux"),				LOCTEXT("LinuxEditor", "Linux (Editor)"),					EPlatformType::Editor,		EPlatformFlags::None,			FPlatformIconPaths(TEXT("Launcher/LinuxTarget/Platform_Linux_24x"), TEXT("Launcher/LinuxTarget/Platform_Linux_128x")),								TEXT(""),												TEXT("Linux_x64"),	EPlatformSDKStatus::Unknown,	TEXT(""),													IsAvailableOnLinux,							false),
-	BuildPlatformInfo(TEXT("LinuxClient"),				TEXT("LinuxClient"),		LOCTEXT("LinuxClient", "Linux (Client-only)"),				EPlatformType::Client,		EPlatformFlags::None,			FPlatformIconPaths(TEXT("Launcher/LinuxTarget/Platform_Linux_24x"), TEXT("Launcher/LinuxTarget/Platform_Linux_128x")),								TEXT(""),												TEXT("Linux_x64"),	EPlatformSDKStatus::Unknown,	TEXT(""),													IsAvailableOnLinux || IsAvailableOnWindows,	false),
-	BuildPlatformInfo(TEXT("LinuxServer"),				TEXT("LinuxServer"),		LOCTEXT("LinuxServer", "Linux (Dedicated Server)"),			EPlatformType::Server,		EPlatformFlags::None,			FPlatformIconPaths(TEXT("Launcher/LinuxTarget/Platform_Linux_24x"), TEXT("Launcher/LinuxTarget/Platform_Linux_128x")),								TEXT(""),												TEXT("Linux_x64"),	EPlatformSDKStatus::Unknown,	TEXT(""),													IsAvailableOnLinux || IsAvailableOnWindows,	false),
-
-	BuildPlatformInfo(TEXT("IOS"),						TEXT("IOS"),				LOCTEXT("IOS", "iOS"),										EPlatformType::Game,		EPlatformFlags::None,			FPlatformIconPaths(TEXT("Launcher/iOSTarget/Platform_iOS_24x"), TEXT("Launcher/iOSTarget/Platform_iOS_128x")),										TEXT(""),												TEXT(""),			EPlatformSDKStatus::Unknown,	IOSTutorial,												IsAvailableOnWindows || IsAvailableOnMac,	true),
-
-	BuildPlatformInfo(TEXT("Android"),					TEXT("Android"),			LOCTEXT("Android", "Android"),								EPlatformType::Game,		EPlatformFlags::None,			FPlatformIconPaths(TEXT("Launcher/AndroidTarget/Platform_Android_24x"), TEXT("Launcher/AndroidTarget/Platform_Android_128x")),						TEXT(""),												TEXT("Android"),	EPlatformSDKStatus::Unknown,	TEXT("Shared/Tutorials/SettingUpAndroidTutorial"),			IsAvailableOnWindows,						true),
-	BuildPlatformInfo(TEXT("Android_All"),				TEXT("Android"),			LOCTEXT("Android_All", "Android (All)"),					EPlatformType::Game,		EPlatformFlags::CookFlavor,		FPlatformIconPaths(TEXT("Launcher/AndroidTarget/Platform_Android_24x"), TEXT("Launcher/AndroidTarget/Platform_Android_128x")),						TEXT(""),												TEXT("Android"),	EPlatformSDKStatus::Unknown,	TEXT("Shared/Tutorials/SettingUpAndroidTutorial"),			IsAvailableOnWindows,						true),
-	BuildPlatformInfo(TEXT("Android_ATC"),				TEXT("Android_ATC"),		LOCTEXT("Android_ATC", "Android (ATC)"),					EPlatformType::Game,		EPlatformFlags::CookFlavor,		FPlatformIconPaths(TEXT("Launcher/AndroidTarget/Platform_Android_ATC_24x"), TEXT("Launcher/AndroidTarget/Platform_Android_128x")),					TEXT("-targetplatform=Android -cookflavor=ATC"),		TEXT("Android"),	EPlatformSDKStatus::Unknown,	TEXT("Shared/Tutorials/SettingUpAndroidTutorial"),			IsAvailableOnWindows,						true),
-	BuildPlatformInfo(TEXT("Android_DXT"),				TEXT("Android_DXT"),		LOCTEXT("Android_DXT", "Android (DXT)"),					EPlatformType::Game,		EPlatformFlags::CookFlavor,		FPlatformIconPaths(TEXT("Launcher/AndroidTarget/Platform_Android_DXT_24x"), TEXT("Launcher/AndroidTarget/Platform_Android_128x")),					TEXT("-targetplatform=Android -cookflavor=DXT"),		TEXT("Android"),	EPlatformSDKStatus::Unknown,	TEXT("Shared/Tutorials/SettingUpAndroidTutorial"),			IsAvailableOnWindows,						true),
-	BuildPlatformInfo(TEXT("Android_ETC1"),				TEXT("Android_ETC1"),		LOCTEXT("Android_ETC1", "Android (ETC1)"),					EPlatformType::Game,		EPlatformFlags::CookFlavor,		FPlatformIconPaths(TEXT("Launcher/AndroidTarget/Platform_Android_ETC1_24x"), TEXT("Launcher/AndroidTarget/Platform_Android_128x")),					TEXT("-targetplatform=Android -cookflavor=ETC1"),		TEXT("Android"),	EPlatformSDKStatus::Unknown,	TEXT("Shared/Tutorials/SettingUpAndroidTutorial"),			IsAvailableOnWindows,						true),
-	BuildPlatformInfo(TEXT("Android_ETC2"),				TEXT("Android_ETC2"),		LOCTEXT("Android_ETC2", "Android (ETC2)"),					EPlatformType::Game,		EPlatformFlags::CookFlavor,		FPlatformIconPaths(TEXT("Launcher/AndroidTarget/Platform_Android_ETC2_24x"), TEXT("Launcher/AndroidTarget/Platform_Android_128x")),					TEXT("-targetplatform=Android -cookflavor=ETC2"),		TEXT("Android"),	EPlatformSDKStatus::Unknown,	TEXT("Shared/Tutorials/SettingUpAndroidTutorial"),			IsAvailableOnWindows,						true),
-	BuildPlatformInfo(TEXT("Android_PVRTC"),			TEXT("Android_PVRTC"),		LOCTEXT("Android_PVRTC", "Android (PVRTC)"),				EPlatformType::Game,		EPlatformFlags::CookFlavor,		FPlatformIconPaths(TEXT("Launcher/AndroidTarget/Platform_Android_PVRTC_24x"), TEXT("Launcher/AndroidTarget/Platform_Android_128x")),				TEXT("-targetplatform=Android -cookflavor=PVRTC"),		TEXT("Android"),	EPlatformSDKStatus::Unknown,	TEXT("Shared/Tutorials/SettingUpAndroidTutorial"),			IsAvailableOnWindows,						true),
-
-	BuildPlatformInfo(TEXT("HTML5"),					TEXT("HTML5"),				LOCTEXT("HTML5", "HTML5"),									EPlatformType::Game,		EPlatformFlags::None,			FPlatformIconPaths(TEXT("Launcher/HTML5Target/Platform_HTML5_24x"), TEXT("Launcher/HTML5Target/Platform_HTML5_128x")),								TEXT(""),												TEXT("HTML5"),		EPlatformSDKStatus::Unknown,	TEXT("Platforms/HTML5/GettingStarted"),						IsAvailableOnWindows || IsAvailableOnMac,   IsAvailableOnWindows),
-
-	BuildPlatformInfo(TEXT("PS4"),						TEXT("PS4"),				LOCTEXT("PS4", "PlayStation 4"),							EPlatformType::Game,		EPlatformFlags::None,			FPlatformIconPaths(TEXT("Launcher/PS4Target/Platform_PS4_24x"), TEXT("Launcher/PS4Target/Platform_PS4_128x")),										TEXT(""),												TEXT("PS4"),		EPlatformSDKStatus::Unknown,	TEXT("Platforms/PS4/GettingStarted"),						IsAvailableOnWindows,						false),
-
-	BuildPlatformInfo(TEXT("XboxOne"),					TEXT("XboxOne"),			LOCTEXT("XboxOne", "Xbox One"),								EPlatformType::Game,		EPlatformFlags::None,			FPlatformIconPaths(TEXT("Launcher/XboxOneTarget/Platform_XboxOne_24x"), TEXT("Launcher/XboxOneTarget/Platform_XboxOne_128x")),						TEXT(""),												TEXT("XboxOne"),	EPlatformSDKStatus::Unknown,	TEXT("Platforms/XboxOne/GettingStarted"),					IsAvailableOnWindows,						false),
+	//				  PlatformInfoName					TargetPlatformName			DisplayName													PlatformType				PlatformFlags					IconPaths																																			UATCommandLine											AutoSDKPath			SDKStatus						SDKTutorial																								EnabledForUse									EnabledForBinary
+	BuildPlatformInfo(TEXT("WindowsNoEditor"),			TEXT("WindowsNoEditor"),	LOCTEXT("WindowsNoEditor", "Windows"),						EPlatformType::Game,		EPlatformFlags::None,			FPlatformIconPaths(TEXT("Launcher/WindowsTarget/Platform_WindowsNoEditor_24x"), TEXT("Launcher/WindowsTarget/Platform_WindowsNoEditor_128x")),		TEXT(""),												TEXT(""),			EPlatformSDKStatus::Unknown,	TEXT("/Engine/Tutorial/Installation/InstallingVisualStudioTutorial.InstallingVisualStudioTutorial"),		IsAvailableOnWindows,						true),
+	BuildPlatformInfo(TEXT("WindowsNoEditor_Win32"),	TEXT("WindowsNoEditor"),	LOCTEXT("WindowsNoEditor_Win32", "Windows (32-bit)"),		EPlatformType::Game,		EPlatformFlags::BuildFlavor,	FPlatformIconPaths(TEXT("Launcher/WindowsTarget/Platform_WindowsNoEditor_24x"), TEXT("Launcher/WindowsTarget/Platform_WindowsNoEditor_128x")),		TEXT("-targetplatform=Win32"),							TEXT(""),			EPlatformSDKStatus::Unknown,	TEXT("/Engine/Tutorial/Installation/InstallingVisualStudioTutorial.InstallingVisualStudioTutorial"),		IsAvailableOnWindows,						true),
+	BuildPlatformInfo(TEXT("WindowsNoEditor_Win64"),	TEXT("WindowsNoEditor"),	LOCTEXT("WindowsNoEditor_Win64", "Windows (64-bit)"),		EPlatformType::Game,		EPlatformFlags::BuildFlavor,	FPlatformIconPaths(TEXT("Launcher/WindowsTarget/Platform_WindowsNoEditor_24x"), TEXT("Launcher/WindowsTarget/Platform_WindowsNoEditor_128x")),		TEXT("-targetplatform=Win64"),							TEXT(""),			EPlatformSDKStatus::Unknown,	TEXT("/Engine/Tutorial/Installation/InstallingVisualStudioTutorial.InstallingVisualStudioTutorial"),		IsAvailableOnWindows,						true),
+	BuildPlatformInfo(TEXT("Windows"),					TEXT("Windows"),			LOCTEXT("WindowsEditor", "Windows (Editor)"),				EPlatformType::Editor,		EPlatformFlags::None,			FPlatformIconPaths(TEXT("Launcher/WindowsTarget/Platform_Windows_24x"), TEXT("Launcher/WindowsTarget/Platform_Windows_128x")),						TEXT(""),												TEXT(""),			EPlatformSDKStatus::Unknown,	TEXT("/Engine/Tutorial/Installation/InstallingVisualStudioTutorial.InstallingVisualStudioTutorial"),		IsAvailableOnWindows,						true),
+	BuildPlatformInfo(TEXT("WindowsClient"),			TEXT("WindowsClient"),		LOCTEXT("WindowsClient", "Windows (Client-only)"),			EPlatformType::Client,		EPlatformFlags::None,			FPlatformIconPaths(TEXT("Launcher/WindowsTarget/Platform_Windows_24x"), TEXT("Launcher/WindowsTarget/Platform_Windows_128x")),						TEXT(""),												TEXT(""),			EPlatformSDKStatus::Unknown,	TEXT("/Engine/Tutorial/Installation/InstallingVisualStudioTutorial.InstallingVisualStudioTutorial"),		IsAvailableOnWindows,						true),
+	BuildPlatformInfo(TEXT("WindowsServer"),			TEXT("WindowsServer"),		LOCTEXT("WindowsServer", "Windows (Dedicated Server)"),		EPlatformType::Server,		EPlatformFlags::None,			FPlatformIconPaths(TEXT("Launcher/WindowsTarget/Platform_WindowsServer_24x"), TEXT("Launcher/WindowsTarget/Platform_WindowsServer_128x")),			TEXT(""),												TEXT(""),			EPlatformSDKStatus::Unknown,	TEXT("/Engine/Tutorial/Installation/InstallingVisualStudioTutorial.InstallingVisualStudioTutorial"),		IsAvailableOnWindows,						true),
+
+	BuildPlatformInfo(TEXT("MacNoEditor"),				TEXT("MacNoEditor"),		LOCTEXT("MacNoEditor", "Mac"),								EPlatformType::Game,		EPlatformFlags::None,			FPlatformIconPaths(TEXT("Launcher/MacTarget/Platform_Mac_24x"), TEXT("Launcher/MacTarget/Platform_Mac_128x")),										TEXT(""),												TEXT(""),			EPlatformSDKStatus::Unknown,	TEXT("/Engine/Tutorial/Installation/InstallingXCodeTutorial.InstallingXCodeTutorial"),						IsAvailableOnMac,							true),
+	BuildPlatformInfo(TEXT("Mac"),						TEXT("Mac"),				LOCTEXT("MacEditor", "Mac (Editor)"),						EPlatformType::Editor,		EPlatformFlags::None,			FPlatformIconPaths(TEXT("Launcher/MacTarget/Platform_Mac_24x"), TEXT("Launcher/MacTarget/Platform_Mac_128x")),										TEXT(""),												TEXT(""),			EPlatformSDKStatus::Unknown,	TEXT("/Engine/Tutorial/Installation/InstallingXCodeTutorial.InstallingXCodeTutorial"),						IsAvailableOnMac,							true),
+	BuildPlatformInfo(TEXT("MacClient"),				TEXT("MacClient"),			LOCTEXT("MacClient", "Mac (Client-only)"),					EPlatformType::Client,		EPlatformFlags::None,			FPlatformIconPaths(TEXT("Launcher/MacTarget/Platform_Mac_24x"), TEXT("Launcher/MacTarget/Platform_Mac_128x")),										TEXT(""),												TEXT(""),			EPlatformSDKStatus::Unknown,	TEXT("/Engine/Tutorial/Installation/InstallingXCodeTutorial.InstallingXCodeTutorial"),						IsAvailableOnMac,							true),
+	BuildPlatformInfo(TEXT("MacServer"),				TEXT("MacServer"),			LOCTEXT("MacServer", "Mac (Dedicated Server)"),				EPlatformType::Server,		EPlatformFlags::None,			FPlatformIconPaths(TEXT("Launcher/MacTarget/Platform_Mac_24x"), TEXT("Launcher/MacTarget/Platform_Mac_128x")),										TEXT(""),												TEXT(""),			EPlatformSDKStatus::Unknown,	TEXT("/Engine/Tutorial/Installation/InstallingXCodeTutorial.InstallingXCodeTutorial"),						IsAvailableOnMac,							true),
+
+	BuildPlatformInfo(TEXT("LinuxNoEditor"),			TEXT("LinuxNoEditor"),		LOCTEXT("LinuxNoEditor", "Linux"),							EPlatformType::Game,		EPlatformFlags::None,			FPlatformIconPaths(TEXT("Launcher/LinuxTarget/Platform_Linux_24x"), TEXT("Launcher/LinuxTarget/Platform_Linux_128x")),								TEXT(""),												TEXT("Linux_x64"),	EPlatformSDKStatus::Unknown,	TEXT(""),																									IsAvailableOnLinux || IsAvailableOnWindows,	false),	
+	BuildPlatformInfo(TEXT("Linux"),					TEXT("Linux"),				LOCTEXT("LinuxEditor", "Linux (Editor)"),					EPlatformType::Editor,		EPlatformFlags::None,			FPlatformIconPaths(TEXT("Launcher/LinuxTarget/Platform_Linux_24x"), TEXT("Launcher/LinuxTarget/Platform_Linux_128x")),								TEXT(""),												TEXT("Linux_x64"),	EPlatformSDKStatus::Unknown,	TEXT(""),																									IsAvailableOnLinux,							false),
+	BuildPlatformInfo(TEXT("LinuxClient"),				TEXT("LinuxClient"),		LOCTEXT("LinuxClient", "Linux (Client-only)"),				EPlatformType::Client,		EPlatformFlags::None,			FPlatformIconPaths(TEXT("Launcher/LinuxTarget/Platform_Linux_24x"), TEXT("Launcher/LinuxTarget/Platform_Linux_128x")),								TEXT(""),												TEXT("Linux_x64"),	EPlatformSDKStatus::Unknown,	TEXT(""),																									IsAvailableOnLinux || IsAvailableOnWindows,	false),
+	BuildPlatformInfo(TEXT("LinuxServer"),				TEXT("LinuxServer"),		LOCTEXT("LinuxServer", "Linux (Dedicated Server)"),			EPlatformType::Server,		EPlatformFlags::None,			FPlatformIconPaths(TEXT("Launcher/LinuxTarget/Platform_Linux_24x"), TEXT("Launcher/LinuxTarget/Platform_Linux_128x")),								TEXT(""),												TEXT("Linux_x64"),	EPlatformSDKStatus::Unknown,	TEXT(""),																									IsAvailableOnLinux || IsAvailableOnWindows,	false),
+
+	BuildPlatformInfo(TEXT("IOS"),						TEXT("IOS"),				LOCTEXT("IOS", "iOS"),										EPlatformType::Game,		EPlatformFlags::None,			FPlatformIconPaths(TEXT("Launcher/iOSTarget/Platform_iOS_24x"), TEXT("Launcher/iOSTarget/Platform_iOS_128x")),										TEXT(""),												TEXT(""),			EPlatformSDKStatus::Unknown,	IOSTutorial,																								IsAvailableOnWindows || IsAvailableOnMac,	true),
+
+	BuildPlatformInfo(TEXT("Android"),					TEXT("Android"),			LOCTEXT("Android", "Android"),								EPlatformType::Game,		EPlatformFlags::None,			FPlatformIconPaths(TEXT("Launcher/AndroidTarget/Platform_Android_24x"), TEXT("Launcher/AndroidTarget/Platform_Android_128x")),						TEXT(""),												TEXT("Android"),	EPlatformSDKStatus::Unknown,	TEXT("/Engine/Tutorial/Mobile/SettingUpAndroidTutorial.SettingUpAndroidTutorial"),							IsAvailableOnWindows,						true),
+	BuildPlatformInfo(TEXT("Android_All"),				TEXT("Android"),			LOCTEXT("Android_All", "Android (All)"),					EPlatformType::Game,		EPlatformFlags::CookFlavor,		FPlatformIconPaths(TEXT("Launcher/AndroidTarget/Platform_Android_24x"), TEXT("Launcher/AndroidTarget/Platform_Android_128x")),						TEXT(""),												TEXT("Android"),	EPlatformSDKStatus::Unknown,	TEXT("/Engine/Tutorial/Mobile/SettingUpAndroidTutorial.SettingUpAndroidTutorial"),							IsAvailableOnWindows,						true),
+	BuildPlatformInfo(TEXT("Android_ATC"),				TEXT("Android_ATC"),		LOCTEXT("Android_ATC", "Android (ATC)"),					EPlatformType::Game,		EPlatformFlags::CookFlavor,		FPlatformIconPaths(TEXT("Launcher/AndroidTarget/Platform_Android_ATC_24x"), TEXT("Launcher/AndroidTarget/Platform_Android_128x")),					TEXT("-targetplatform=Android -cookflavor=ATC"),		TEXT("Android"),	EPlatformSDKStatus::Unknown,	TEXT("/Engine/Tutorial/Mobile/SettingUpAndroidTutorial.SettingUpAndroidTutorial"),							IsAvailableOnWindows,						true),
+	BuildPlatformInfo(TEXT("Android_DXT"),				TEXT("Android_DXT"),		LOCTEXT("Android_DXT", "Android (DXT)"),					EPlatformType::Game,		EPlatformFlags::CookFlavor,		FPlatformIconPaths(TEXT("Launcher/AndroidTarget/Platform_Android_DXT_24x"), TEXT("Launcher/AndroidTarget/Platform_Android_128x")),					TEXT("-targetplatform=Android -cookflavor=DXT"),		TEXT("Android"),	EPlatformSDKStatus::Unknown,	TEXT("/Engine/Tutorial/Mobile/SettingUpAndroidTutorial.SettingUpAndroidTutorial"),							IsAvailableOnWindows,						true),
+	BuildPlatformInfo(TEXT("Android_ETC1"),				TEXT("Android_ETC1"),		LOCTEXT("Android_ETC1", "Android (ETC1)"),					EPlatformType::Game,		EPlatformFlags::CookFlavor,		FPlatformIconPaths(TEXT("Launcher/AndroidTarget/Platform_Android_ETC1_24x"), TEXT("Launcher/AndroidTarget/Platform_Android_128x")),					TEXT("-targetplatform=Android -cookflavor=ETC1"),		TEXT("Android"),	EPlatformSDKStatus::Unknown,	TEXT("/Engine/Tutorial/Mobile/SettingUpAndroidTutorial.SettingUpAndroidTutorial"),							IsAvailableOnWindows,						true),
+	BuildPlatformInfo(TEXT("Android_ETC2"),				TEXT("Android_ETC2"),		LOCTEXT("Android_ETC2", "Android (ETC2)"),					EPlatformType::Game,		EPlatformFlags::CookFlavor,		FPlatformIconPaths(TEXT("Launcher/AndroidTarget/Platform_Android_ETC2_24x"), TEXT("Launcher/AndroidTarget/Platform_Android_128x")),					TEXT("-targetplatform=Android -cookflavor=ETC2"),		TEXT("Android"),	EPlatformSDKStatus::Unknown,	TEXT("/Engine/Tutorial/Mobile/SettingUpAndroidTutorial.SettingUpAndroidTutorial"),							IsAvailableOnWindows,						true),
+	BuildPlatformInfo(TEXT("Android_PVRTC"),			TEXT("Android_PVRTC"),		LOCTEXT("Android_PVRTC", "Android (PVRTC)"),				EPlatformType::Game,		EPlatformFlags::CookFlavor,		FPlatformIconPaths(TEXT("Launcher/AndroidTarget/Platform_Android_PVRTC_24x"), TEXT("Launcher/AndroidTarget/Platform_Android_128x")),				TEXT("-targetplatform=Android -cookflavor=PVRTC"),		TEXT("Android"),	EPlatformSDKStatus::Unknown,	TEXT("/Engine/Tutorial/Mobile/SettingUpAndroidTutorial.SettingUpAndroidTutorial"),							IsAvailableOnWindows,						true),
+
+	BuildPlatformInfo(TEXT("HTML5"),					TEXT("HTML5"),				LOCTEXT("HTML5", "HTML5"),									EPlatformType::Game,		EPlatformFlags::None,			FPlatformIconPaths(TEXT("Launcher/HTML5Target/Platform_HTML5_24x"), TEXT("Launcher/HTML5Target/Platform_HTML5_128x")),								TEXT(""),												TEXT("HTML5"),		EPlatformSDKStatus::Unknown,	TEXT(""),																									IsAvailableOnWindows || IsAvailableOnMac,   IsAvailableOnWindows),
+
+	BuildPlatformInfo(TEXT("PS4"),						TEXT("PS4"),				LOCTEXT("PS4", "PlayStation 4"),							EPlatformType::Game,		EPlatformFlags::None,			FPlatformIconPaths(TEXT("Launcher/PS4Target/Platform_PS4_24x"), TEXT("Launcher/PS4Target/Platform_PS4_128x")),										TEXT(""),												TEXT("PS4"),		EPlatformSDKStatus::Unknown,	TEXT(""),																									IsAvailableOnWindows,						false),
+
+	BuildPlatformInfo(TEXT("XboxOne"),					TEXT("XboxOne"),			LOCTEXT("XboxOne", "Xbox One"),								EPlatformType::Game,		EPlatformFlags::None,			FPlatformIconPaths(TEXT("Launcher/XboxOneTarget/Platform_XboxOne_24x"), TEXT("Launcher/XboxOneTarget/Platform_XboxOne_128x")),						TEXT(""),												TEXT("XboxOne"),	EPlatformSDKStatus::Unknown,	TEXT(""),																									IsAvailableOnWindows,						false),
 
 	//BuildPlatformInfo(TEXT("WinRT"),					TEXT("WinRT"),				LOCTEXT("WinRT", "Windows RT"),								EPlatformType::Game,		EPlatformFlags::None,			FPlatformIconPaths(TEXT("Launcher/WindowsTarget/Platform_WindowsNoEditor_24x"), TEXT("Launcher/WindowsTarget/Platform_WindowsNoEditor_128x")),		TEXT("")),
 };
diff --git a/Engine/Source/Developer/iOS/IOSTargetPlatform/Private/IOSTargetPlatform.cpp b/Engine/Source/Developer/iOS/IOSTargetPlatform/Private/IOSTargetPlatform.cpp
--- a/Engine/Source/Developer/iOS/IOSTargetPlatform/Private/IOSTargetPlatform.cpp
+++ b/Engine/Source/Developer/iOS/IOSTargetPlatform/Private/IOSTargetPlatform.cpp
@@ -76,7 +76,7 @@ ITargetDevicePtr FIOSTargetPlatform::GetDevice( const FTargetDeviceId& DeviceId
 }
 
 
-bool FIOSTargetPlatform::IsSdkInstalled(bool bProjectHasCode, FString& OutDocumentationPath) const
+bool FIOSTargetPlatform::IsSdkInstalled(bool bProjectHasCode, FString& OutTutorialPath) const
 {
 	bool biOSSDKInstalled = true; // @todo How do we check that the iOS SDK is installed when building from Windows? Is that even possible?
 #if PLATFORM_MAC
@@ -90,7 +90,7 @@ bool FIOSTargetPlatform::IsSdkInstalled(bool bProjectHasCode, FString& OutDocume
 		unsigned long pathSize = 256;
 		if (lRes != ERROR_SUCCESS || RegQueryValueEx(hKey, TEXT("iTunesMobileDeviceDLL"), 0, NULL, (BYTE*)dllPath, &pathSize) != ERROR_SUCCESS || IFileManager::Get().FileSize(*FString(dllPath)) == INDEX_NONE)
 		{
-			OutDocumentationPath = FString("Shared/Tutorials/InstallingiTunesTutorial");
+			OutTutorialPath = FString("/Engine/Tutorial/Mobile/InstallingiTunesTutorial.InstallingiTunesTutorial");
 			biOSSDKInstalled = false;
 		}
 	}
@@ -106,24 +106,24 @@ static void OnOutput(FString Message)
 	UE_LOG(LogTemp, Display, TEXT("%s"), *Message);
 }
 
-int FIOSTargetPlatform::DoesntHaveRequirements(const FString& ProjectPath, bool bProjectHasCode, FString& OutDocumentationPath) const
+int FIOSTargetPlatform::DoesntHaveRequirements(const FString& ProjectPath, bool bProjectHasCode, FString& OutTutorialPath) const
 {
 	int bReadyToBuild = ETargetPlatformReadyStatus::Ready; // @todo How do we check that the iOS SDK is installed when building from Windows? Is that even possible?
-	if (!IsSdkInstalled(bProjectHasCode, OutDocumentationPath))
+	if (!IsSdkInstalled(bProjectHasCode, OutTutorialPath))
 	{
 		bReadyToBuild |= ETargetPlatformReadyStatus::SDKNotFound;
 	}
 #if PLATFORM_MAC
-	OutDocumentationPath = FString("Shared/Tutorials/InstallingXCodeTutorial");
+	OutTutorialPath = FString("/Engine/Tutorial/Installation/InstallingXCodeTutorial.InstallingXCodeTutorial");
 #else
 	if (bProjectHasCode && FRocketSupport::IsRocket())
 	{
-		OutDocumentationPath = FString("Shared/Tutorials/iOSonPCRestrictions");
+		OutTutorialPath = FString("/Engine/Tutorial/Mobile/iOSonPCRestrictions.iOSonPCRestrictions");
 		bReadyToBuild |= ETargetPlatformReadyStatus::CodeUnsupported;
 	}
 	if (FRocketSupport::IsRocket() && IProjectManager::Get().IsNonDefaultPluginEnabled())
 	{
-		OutDocumentationPath = FString("Shared/Tutorials/iOSonPCValidPlugins");
+		OutTutorialPath = FString("/Engine/Tutorial/Mobile/iOSonPCValidPlugins.iOSonPCValidPlugins");
 		bReadyToBuild |= ETargetPlatformReadyStatus::PluginsUnsupported;
 	}
 
@@ -141,23 +141,23 @@ int FIOSTargetPlatform::DoesntHaveRequirements(const FString& ProjectPath, bool
 	int RetCode = IPPProcess->GetReturnCode();
 	if (RetCode == 14)
 	{
-		OutDocumentationPath = FString("Shared/Tutorials/CreatingInfoPlist");
+		OutTutorialPath = FString("/Engine/Tutorial/Mobile/CreatingInfoPlist.CreatingInfoPlist");
 		bReadyToBuild |= ETargetPlatformReadyStatus::ManifestNotFound;
 	}
 	else if (RetCode == 13)
 	{
-		OutDocumentationPath = FString("Shared/Tutorials/CreatingSigningCertAndProvision");
+		OutTutorialPath = FString("/Engine/Tutorial/Mobile/CreatingSigningCertAndProvisionTutorial.CreatingSigningCertAndProvisionTutorial");
 		bReadyToBuild |= ETargetPlatformReadyStatus::SigningKeyNotFound;
 		bReadyToBuild |= ETargetPlatformReadyStatus::ProvisionNotFound;
 	}
 	else if (RetCode == 12)
 	{
-		OutDocumentationPath = FString("Shared/Tutorials/CreatingSigningCertAndProvision");
+		OutTutorialPath = FString("/Engine/Tutorial/Mobile/CreatingSigningCertAndProvisionTutorial.CreatingSigningCertAndProvisionTutorial");
 		bReadyToBuild |= ETargetPlatformReadyStatus::SigningKeyNotFound;
 	}
 	else if (RetCode == 11)
 	{
-		OutDocumentationPath = FString("Shared/Tutorials/CreatingSigningCertAndProvision");
+		OutTutorialPath = FString("/Engine/Tutorial/Mobile/CreatingSigningCertAndProvisionTutorial.CreatingSigningCertAndProvisionTutorial");
 		bReadyToBuild |= ETargetPlatformReadyStatus::ProvisionNotFound;
 	}
 #endif
diff --git a/Engine/Source/Developer/iOS/IOSTargetPlatform/Private/IOSTargetPlatform.h b/Engine/Source/Developer/iOS/IOSTargetPlatform/Private/IOSTargetPlatform.h
--- a/Engine/Source/Developer/iOS/IOSTargetPlatform/Private/IOSTargetPlatform.h
+++ b/Engine/Source/Developer/iOS/IOSTargetPlatform/Private/IOSTargetPlatform.h
@@ -83,8 +83,8 @@ public:
 		return TTargetPlatformBase<FIOSPlatformProperties>::SupportsFeature(Feature);
 	}
 
-	virtual bool IsSdkInstalled(bool bProjectHasCode, FString& OutDocumentationPath) const override;
-	virtual int DoesntHaveRequirements(const FString& ProjectPath, bool bProjectHasCode, FString& OutDocumentationPath) const override;
+	virtual bool IsSdkInstalled(bool bProjectHasCode, FString& OutTutorialPath) const override;
+	virtual int DoesntHaveRequirements(const FString& ProjectPath, bool bProjectHasCode, FString& OutTutorialPath) const override;
 
 
 #if WITH_ENGINE
diff --git a/Engine/Source/Editor/IntroTutorials/Private/InteractivityData.cpp b/Engine/Source/Editor/IntroTutorials/Private/InteractivityData.cpp
--- a/Engine/Source/Editor/IntroTutorials/Private/InteractivityData.cpp
+++ b/Engine/Source/Editor/IntroTutorials/Private/InteractivityData.cpp
@@ -541,7 +541,7 @@ void PopulateInteractivityData(TSharedPtr<FInteractiveTutorials> InteractiveTuto
 		}
 	};
 
-	InteractiveTutorials->BindInteractivityData(FIntroTutorials::InEditorTutorialPath, FBindInteractivityData::CreateStatic(&FInEditorTutorial::BindInteractivityData));
+//	InteractiveTutorials->BindInteractivityData(FIntroTutorials::InEditorTutorialPath, FBindInteractivityData::CreateStatic(&FInEditorTutorial::BindInteractivityData));
 
 	//////////////////////////////////////////////////////////////////////////
 
@@ -833,7 +833,7 @@ void PopulateInteractivityData(TSharedPtr<FInteractiveTutorials> InteractiveTuto
 		}
 	};
 
-	InteractiveTutorials->BindInteractivityData(TEXT("Shared/Tutorials/inEditorTutorial_2"), FBindInteractivityData::CreateStatic(&FProofOfConceptTutorial_2::BindInteractivityData));
+//	InteractiveTutorials->BindInteractivityData(TEXT("Shared/Tutorials/inEditorTutorial_2"), FBindInteractivityData::CreateStatic(&FProofOfConceptTutorial_2::BindInteractivityData));
 
 	//////////////////////////////////////////////////////////////////////////
 
@@ -1449,7 +1449,7 @@ void PopulateInteractivityData(TSharedPtr<FInteractiveTutorials> InteractiveTuto
 		}
 	};
 
-	InteractiveTutorials->BindInteractivityData(FIntroTutorials::InEditorGamifiedTutorialPath, FBindInteractivityData::CreateStatic(&FGamifiedTutorial::BindInteractivityData));
+//	InteractiveTutorials->BindInteractivityData(FIntroTutorials::InEditorGamifiedTutorialPath, FBindInteractivityData::CreateStatic(&FGamifiedTutorial::BindInteractivityData));
 
 	//////////////////////////////////////////////////////////////////////////
 	// Force Save
@@ -1471,7 +1471,7 @@ void PopulateInteractivityData(TSharedPtr<FInteractiveTutorials> InteractiveTuto
 		}
 	};
 
-	InteractiveTutorials->BindInteractivityData(TEXT("Shared/Tutorials/inEditorTutorialManipulation"), FBindInteractivityData::CreateStatic(&FManipulationTutorial::BindInteractivityData));
+//	InteractiveTutorials->BindInteractivityData(TEXT("Shared/Tutorials/inEditorTutorialManipulation"), FBindInteractivityData::CreateStatic(&FManipulationTutorial::BindInteractivityData));
 
 
 	//////////////////////////////////////////////////////////////////////////
@@ -1510,7 +1510,7 @@ void PopulateInteractivityData(TSharedPtr<FInteractiveTutorials> InteractiveTuto
 		}
 	};
 
-	InteractiveTutorials->BindInteractivityData(TEXT("Shared/Tutorials/BlueprintInterfaceTutorial"), FBindInteractivityData::CreateStatic(&FBlueprintInterfaceTutorial::BindInteractivityData));
+//	InteractiveTutorials->BindInteractivityData(TEXT("Shared/Tutorials/BlueprintInterfaceTutorial"), FBindInteractivityData::CreateStatic(&FBlueprintInterfaceTutorial::BindInteractivityData));
 
 
 	//////////////////////////////////////////////////////////////////////////
@@ -1539,7 +1539,7 @@ void PopulateInteractivityData(TSharedPtr<FInteractiveTutorials> InteractiveTuto
 		}
 	};
 
-	InteractiveTutorials->BindInteractivityData(TEXT("Shared/Tutorials/BlueprintMacroLibInterfaceTutorial"), FBindInteractivityData::CreateStatic(&FBlueprintMacroLibInterfaceTutorial::BindInteractivityData));
+//	InteractiveTutorials->BindInteractivityData(TEXT("Shared/Tutorials/BlueprintMacroLibInterfaceTutorial"), FBindInteractivityData::CreateStatic(&FBlueprintMacroLibInterfaceTutorial::BindInteractivityData));
 
 
 	//////////////////////////////////////////////////////////////////////////
@@ -1566,7 +1566,7 @@ void PopulateInteractivityData(TSharedPtr<FInteractiveTutorials> InteractiveTuto
 		}
 	};
 
-	InteractiveTutorials->BindInteractivityData(TEXT("Shared/Tutorials/LevelBlueprintInterfaceTutorial"), FBindInteractivityData::CreateStatic(&FLevelBlueprintInterfaceTutorial::BindInteractivityData));
+//	InteractiveTutorials->BindInteractivityData(TEXT("Shared/Tutorials/LevelBlueprintInterfaceTutorial"), FBindInteractivityData::CreateStatic(&FLevelBlueprintInterfaceTutorial::BindInteractivityData));
 
 
 	//////////////////////////////////////////////////////////////////////////
@@ -1588,7 +1588,7 @@ void PopulateInteractivityData(TSharedPtr<FInteractiveTutorials> InteractiveTuto
 		}
 	};
 
-	InteractiveTutorials->BindInteractivityData(TEXT("Shared/Tutorials/BlueprintInterfacesInterfaceTutorial"), FBindInteractivityData::CreateStatic(&FBlueprintInterfacesInterfaceTutorial::BindInteractivityData));
+//	InteractiveTutorials->BindInteractivityData(TEXT("Shared/Tutorials/BlueprintInterfacesInterfaceTutorial"), FBindInteractivityData::CreateStatic(&FBlueprintInterfacesInterfaceTutorial::BindInteractivityData));
 
 
 	//////////////////////////////////////////////////////////////////////////
@@ -1662,7 +1662,7 @@ void PopulateInteractivityData(TSharedPtr<FInteractiveTutorials> InteractiveTuto
 		}
 	};
 
-	InteractiveTutorials->BindInteractivityData(TEXT("Shared/Tutorials/TestHighlight"), FBindInteractivityData::CreateStatic(&FTestHighlightTutorial::BindInteractivityData));
+//	InteractiveTutorials->BindInteractivityData(TEXT("Shared/Tutorials/TestHighlight"), FBindInteractivityData::CreateStatic(&FTestHighlightTutorial::BindInteractivityData));
 
 
 	//////////////////////////////////////////////////////////////////////////
@@ -1703,7 +1703,7 @@ void PopulateInteractivityData(TSharedPtr<FInteractiveTutorials> InteractiveTuto
 		}
 	};
 
-	InteractiveTutorials->BindInteractivityData(TEXT("Shared/Tutorials/EditingObjects"), FBindInteractivityData::CreateStatic(&FEditingObjectsTutorial::BindInteractivityData));
+//	InteractiveTutorials->BindInteractivityData(TEXT("Shared/Tutorials/EditingObjects"), FBindInteractivityData::CreateStatic(&FEditingObjectsTutorial::BindInteractivityData));
 
 	//////////////////////////////////////////////////////////////////////////
 
@@ -1716,7 +1716,7 @@ void PopulateInteractivityData(TSharedPtr<FInteractiveTutorials> InteractiveTuto
 		}
 	};
 
-	InteractiveTutorials->BindInteractivityData(TEXT("Shared/Tutorials/TestCase"), FBindInteractivityData::CreateStatic(&FTestCaseTutorial::BindInteractivityData));
+//	InteractiveTutorials->BindInteractivityData(TEXT("Shared/Tutorials/TestCase"), FBindInteractivityData::CreateStatic(&FTestCaseTutorial::BindInteractivityData));
 
 	//////////////////////////////////////////////////////////////////////////
 
@@ -1728,7 +1728,7 @@ void PopulateInteractivityData(TSharedPtr<FInteractiveTutorials> InteractiveTuto
 		}
 	};
 
-	InteractiveTutorials->BindInteractivityData(FIntroTutorials::HomePath, FBindInteractivityData::CreateStatic(&FTutorialHome::BindInteractivityData));
+//	InteractiveTutorials->BindInteractivityData(FIntroTutorials::HomePath, FBindInteractivityData::CreateStatic(&FTutorialHome::BindInteractivityData));
 
 	//////////////////////////////////////////////////////////////////////////
 
@@ -1757,7 +1757,7 @@ void PopulateInteractivityData(TSharedPtr<FInteractiveTutorials> InteractiveTuto
 		}
 	};
 
-	InteractiveTutorials->BindInteractivityData(FIntroTutorials::UE4WelcomeTutorial.TutorialPath, FBindInteractivityData::CreateStatic(&FWelcomeTutorial::BindInteractivityData));
+//	InteractiveTutorials->BindInteractivityData(FIntroTutorials::UE4WelcomeTutorial.TutorialPath, FBindInteractivityData::CreateStatic(&FWelcomeTutorial::BindInteractivityData));
 
 	//////////////////////////////////////////////////////////////////////////
 
@@ -1769,7 +1769,7 @@ void PopulateInteractivityData(TSharedPtr<FInteractiveTutorials> InteractiveTuto
 		}
 	};
 
-	InteractiveTutorials->BindInteractivityData(TEXT("Shared/Tutorials/inEditorTutorialCamera"), FBindInteractivityData::CreateStatic(&FCameraTutorial::BindInteractivityData));
+//	InteractiveTutorials->BindInteractivityData(TEXT("Shared/Tutorials/inEditorTutorialCamera"), FBindInteractivityData::CreateStatic(&FCameraTutorial::BindInteractivityData));
 
 
 	//////////////////////////////////////////////////////////////////////////
@@ -1783,7 +1783,7 @@ void PopulateInteractivityData(TSharedPtr<FInteractiveTutorials> InteractiveTuto
 		}
 	};
 
-	InteractiveTutorials->BindInteractivityData(TEXT("Shared/Tutorials/InTextureEditorTutorial"), FBindInteractivityData::CreateStatic(&FInTextureEditorTutorial::BindInteractivityData));
+//	InteractiveTutorials->BindInteractivityData(TEXT("Shared/Tutorials/InTextureEditorTutorial"), FBindInteractivityData::CreateStatic(&FInTextureEditorTutorial::BindInteractivityData));
 
 
 	//////////////////////////////////////////////////////////////////////////
@@ -1797,7 +1797,7 @@ void PopulateInteractivityData(TSharedPtr<FInteractiveTutorials> InteractiveTuto
 		}
 	};
 
-	InteractiveTutorials->BindInteractivityData(TEXT("Shared/Tutorials/InMaterialEditorTutorial"), FBindInteractivityData::CreateStatic(&FInMaterialEditorTutorial::BindInteractivityData));
+//	InteractiveTutorials->BindInteractivityData(TEXT("Shared/Tutorials/InMaterialEditorTutorial"), FBindInteractivityData::CreateStatic(&FInMaterialEditorTutorial::BindInteractivityData));
 
 
 	//////////////////////////////////////////////////////////////////////////
@@ -1812,7 +1812,7 @@ void PopulateInteractivityData(TSharedPtr<FInteractiveTutorials> InteractiveTuto
 		}
 	};
 
-	InteractiveTutorials->BindInteractivityData(TEXT("Shared/Tutorials/InParticleSystemEditorTutorial"), FBindInteractivityData::CreateStatic(&FInParticleSystemEditorTutorial::BindInteractivityData));
+//	InteractiveTutorials->BindInteractivityData(TEXT("Shared/Tutorials/InParticleSystemEditorTutorial"), FBindInteractivityData::CreateStatic(&FInParticleSystemEditorTutorial::BindInteractivityData));
 
 
 
@@ -1828,7 +1828,7 @@ void PopulateInteractivityData(TSharedPtr<FInteractiveTutorials> InteractiveTuto
 		}
 	};
 
-	InteractiveTutorials->BindInteractivityData(TEXT("Shared/Tutorials/InPersonaAnimEditorTutorial"), FBindInteractivityData::CreateStatic(&FInPersonaAnimEditorTutorial::BindInteractivityData));
+//	InteractiveTutorials->BindInteractivityData(TEXT("Shared/Tutorials/InPersonaAnimEditorTutorial"), FBindInteractivityData::CreateStatic(&FInPersonaAnimEditorTutorial::BindInteractivityData));
 
 	
 
@@ -1844,7 +1844,7 @@ void PopulateInteractivityData(TSharedPtr<FInteractiveTutorials> InteractiveTuto
 		}
 	};
 
-	InteractiveTutorials->BindInteractivityData(TEXT("Shared/Tutorials/InPersonaAnimBlueprintEditorTutorial"), FBindInteractivityData::CreateStatic(&FInPersonaAnimBlueprintEditorTutorial::BindInteractivityData));
+//	InteractiveTutorials->BindInteractivityData(TEXT("Shared/Tutorials/InPersonaAnimBlueprintEditorTutorial"), FBindInteractivityData::CreateStatic(&FInPersonaAnimBlueprintEditorTutorial::BindInteractivityData));
 
 
 
@@ -1860,7 +1860,7 @@ void PopulateInteractivityData(TSharedPtr<FInteractiveTutorials> InteractiveTuto
 		}
 	};
 
-	InteractiveTutorials->BindInteractivityData(TEXT("Shared/Tutorials/InDestructibleMeshEditorTutorial"), FBindInteractivityData::CreateStatic(&FInDestructibleMeshEditorTutorial::BindInteractivityData));
+//	InteractiveTutorials->BindInteractivityData(TEXT("Shared/Tutorials/InDestructibleMeshEditorTutorial"), FBindInteractivityData::CreateStatic(&FInDestructibleMeshEditorTutorial::BindInteractivityData));
 
 
 
@@ -1876,7 +1876,7 @@ void PopulateInteractivityData(TSharedPtr<FInteractiveTutorials> InteractiveTuto
 		}
 	};
 
-	InteractiveTutorials->BindInteractivityData(TEXT("Shared/Tutorials/InSkeletalMeshEditorTutorial"), FBindInteractivityData::CreateStatic(&FInSkeletalMeshEditorTutorial::BindInteractivityData));
+//	InteractiveTutorials->BindInteractivityData(TEXT("Shared/Tutorials/InSkeletalMeshEditorTutorial"), FBindInteractivityData::CreateStatic(&FInSkeletalMeshEditorTutorial::BindInteractivityData));
 
 
 
@@ -1892,7 +1892,7 @@ void PopulateInteractivityData(TSharedPtr<FInteractiveTutorials> InteractiveTuto
 		}
 	};
 
-	InteractiveTutorials->BindInteractivityData(TEXT("Shared/Tutorials/InStaticMeshEditorTutorial"), FBindInteractivityData::CreateStatic(&FInStaticMeshEditorTutorial::BindInteractivityData));
+//	InteractiveTutorials->BindInteractivityData(TEXT("Shared/Tutorials/InStaticMeshEditorTutorial"), FBindInteractivityData::CreateStatic(&FInStaticMeshEditorTutorial::BindInteractivityData));
 
 
 
@@ -1908,7 +1908,7 @@ void PopulateInteractivityData(TSharedPtr<FInteractiveTutorials> InteractiveTuto
 		}
 	};
 
-	InteractiveTutorials->BindInteractivityData(TEXT("Shared/Tutorials/InSoundCueEditorTutorial"), FBindInteractivityData::CreateStatic(&FInSoundCueEditorTutorial::BindInteractivityData));
+//	InteractiveTutorials->BindInteractivityData(TEXT("Shared/Tutorials/InSoundCueEditorTutorial"), FBindInteractivityData::CreateStatic(&FInSoundCueEditorTutorial::BindInteractivityData));
 
 
 
@@ -1926,7 +1926,7 @@ void PopulateInteractivityData(TSharedPtr<FInteractiveTutorials> InteractiveTuto
 		}
 	};
 
-	InteractiveTutorials->BindInteractivityData(TEXT("Shared/Tutorials/inEditorTutorialPlacement"), FBindInteractivityData::CreateStatic(&FPlacementTutorial::BindInteractivityData));
+//	InteractiveTutorials->BindInteractivityData(TEXT("Shared/Tutorials/inEditorTutorialPlacement"), FBindInteractivityData::CreateStatic(&FPlacementTutorial::BindInteractivityData));
 
 	//////////////////////////////////////////////////////////////////////////
 
@@ -1938,7 +1938,7 @@ void PopulateInteractivityData(TSharedPtr<FInteractiveTutorials> InteractiveTuto
 		}
 	};
 
-	InteractiveTutorials->BindInteractivityData(TEXT("Shared/Tutorials/InBlueprintEditorTutorial"), FBindInteractivityData::CreateStatic(&FBlueprintSplash::BindInteractivityData));
+//	InteractiveTutorials->BindInteractivityData(TEXT("Shared/Tutorials/InBlueprintEditorTutorial"), FBindInteractivityData::CreateStatic(&FBlueprintSplash::BindInteractivityData));
 
 	//////////////////////////////////////////////////////////////////////////
 
@@ -1950,7 +1950,7 @@ void PopulateInteractivityData(TSharedPtr<FInteractiveTutorials> InteractiveTuto
 		}
 	};
 
-	InteractiveTutorials->BindInteractivityData(TEXT("Shared/Tutorials/AddCodeToProjectTutorial"), FBindInteractivityData::CreateStatic(&FAddCodeToProjectTutorial::BindInteractivityData));
+//	InteractiveTutorials->BindInteractivityData(TEXT("Shared/Tutorials/AddCodeToProjectTutorial"), FBindInteractivityData::CreateStatic(&FAddCodeToProjectTutorial::BindInteractivityData));
 
 	//////////////////////////////////////////////////////////////////////////
 
@@ -1977,7 +1977,7 @@ void PopulateInteractivityData(TSharedPtr<FInteractiveTutorials> InteractiveTuto
 		}
 	};
 
-	InteractiveTutorials->BindInteractivityData(TEXT("Shared/Tutorials/ContentBrowserWalkthrough"), FBindInteractivityData::CreateStatic(&FContentBrowserWalkthrough::BindInteractivityData));
+//	InteractiveTutorials->BindInteractivityData(TEXT("Shared/Tutorials/ContentBrowserWalkthrough"), FBindInteractivityData::CreateStatic(&FContentBrowserWalkthrough::BindInteractivityData));
 
 	//////////////////////////////////////////////////////////////////////////
 
@@ -1989,7 +1989,7 @@ void PopulateInteractivityData(TSharedPtr<FInteractiveTutorials> InteractiveTuto
 		}
 	};
 
-	InteractiveTutorials->BindInteractivityData(TEXT("Shared/Tutorials/TemplateTutorials/TemplateOverview"), FBindInteractivityData::CreateStatic(&FTemplateOverview::BindInteractivityData));
+//	InteractiveTutorials->BindInteractivityData(TEXT("Shared/Tutorials/TemplateTutorials/TemplateOverview"), FBindInteractivityData::CreateStatic(&FTemplateOverview::BindInteractivityData));
 	
 	//////////////////////////////////////////////////////////////////////////
 
@@ -2021,7 +2021,7 @@ void PopulateInteractivityData(TSharedPtr<FInteractiveTutorials> InteractiveTuto
 		}
 	};
 
-	InteractiveTutorials->BindInteractivityData(TEXT("Shared/Tutorials/TemplateTutorials/FirstPerson"), FBindInteractivityData::CreateStatic(&FFirstPersonTemplateTutorial::BindInteractivityData));
+//	InteractiveTutorials->BindInteractivityData(TEXT("Shared/Tutorials/TemplateTutorials/FirstPerson"), FBindInteractivityData::CreateStatic(&FFirstPersonTemplateTutorial::BindInteractivityData));
 
 }
 
diff --git a/Engine/Source/Editor/IntroTutorials/Private/IntroTutorials.cpp b/Engine/Source/Editor/IntroTutorials/Private/IntroTutorials.cpp
--- a/Engine/Source/Editor/IntroTutorials/Private/IntroTutorials.cpp
+++ b/Engine/Source/Editor/IntroTutorials/Private/IntroTutorials.cpp
@@ -28,56 +28,33 @@
 #include "EditorTutorialDetailsCustomization.h"
 #include "STutorialRoot.h"
 #include "STutorialButton.h"
+#include "ToolkitManager.h"
+#include "BlueprintEditor.h"
 
 #define LOCTEXT_NAMESPACE "IntroTutorials"
 
 IMPLEMENT_MODULE(FIntroTutorials, IntroTutorials)
 
-const FName FIntroTutorials::IntroTutorialTabName(TEXT("IntroTutorial"));
-
-const FString FIntroTutorials::IntroTutorialConfigSection(TEXT("IntroTutorials"));
-const FString FIntroTutorials::DisableTutorialsSettingName(TEXT("DisableAllTutorials"));
-const FString FIntroTutorials::InEditorTutorialPath(TEXT("Shared/Tutorials/inEditorTutorial"));
-const FString FIntroTutorials::WelcomeTutorialPath(TEXT("Shared/Tutorials/UE4Welcome"));
-const FString FIntroTutorials::InEditorGamifiedTutorialPath(TEXT("Shared/Tutorials/inEditorGamifiedTutorial"));
-const FString FIntroTutorials::HomePath(TEXT("Shared/Tutorials"));
-const FString FIntroTutorials::BlueprintHomePath(TEXT("Shared/Tutorials/TemplateTutorials/TemplateOverview"));
-const FString FIntroTutorials::TemplateOverviewPath(TEXT("Shared/Tutorials/UE4Welcome"));
-
-const FWelcomeTutorialProperties FIntroTutorials::UE4WelcomeTutorial(TEXT("Shared/Tutorials/UE4Welcome"), TEXT("SeenUE4Welcome"));	
-const FWelcomeTutorialProperties FIntroTutorials::BlueprintHomeTutorial(TEXT("Shared/Tutorials/InBlueprintEditorTutorial"), TEXT("SeenBlueprintWelcome"));
-const FWelcomeTutorialProperties FIntroTutorials::ClassBlueprintWelcomeTutorial(TEXT("Shared/Tutorials/BlueprintInterfaceTutorial"), TEXT("SeenBlueprintWelcome_Class"), FString("5DCD4D58-9C9F-407F-8388-75D89CBBA7E8"));
-const FWelcomeTutorialProperties FIntroTutorials::MacroLibraryBlueprintWelcomeTutorial(TEXT("Shared/Tutorials/BlueprintMacroLibInterfaceTutorial"), TEXT("SeenBlueprintWelcome_MacroLib"), FString("0D5081E0-F29A-4C35-B4DC-1E5E2825AB54"));
-const FWelcomeTutorialProperties FIntroTutorials::InterfaceBlueprintWelcomeTutorial(TEXT("Shared/Tutorials/BlueprintInterfacesInterfaceTutorial"), TEXT("SeenBlueprintWelcome_Interface"), FString("E86E8C9A-4804-4680-B10F-BDC8E95C7AFD"));
-const FWelcomeTutorialProperties FIntroTutorials::LevelScriptBlueprintWelcomeTutorial(TEXT("Shared/Tutorials/LevelBlueprintInterfaceTutorial"), TEXT("SeenBlueprintWelcome_LevelScript"), FString("B061E309-517D-4916-BFCB-E8104C8F4C35"));
-const FWelcomeTutorialProperties FIntroTutorials::AddCodeToProjectWelcomeTutorial(TEXT("Shared/Tutorials/AddCodeToProjectTutorial"), TEXT("SeenAddCodeToProjectWelcome"), FString("D8D9A7E7-68B3-4CBE-80FE-4B88C47B7524"));
-const FWelcomeTutorialProperties FIntroTutorials::MatineeEditorWelcomeTutorial(TEXT("Shared/Tutorials/InMatineeEditorTutorial"), TEXT("SeenMatineeEditorWelcome"), FString("6439C991-A77B-4B52-953D-3F29B1DD1860"));
-const FWelcomeTutorialProperties FIntroTutorials::TemplateOverview(TemplateOverviewPath, TEXT("SeenTemplateOverview"));
-
 FIntroTutorials::FIntroTutorials()
 	: CurrentObjectClass(nullptr)
 	, ContentIntroCurve(nullptr)
 {
 	bDisableTutorials = false;
-	bEnablePostTutorialSurveys = false;
 
-	bDesireResettingTutorialSeenFlagOnLoad = FParse::Param(FCommandLine::Get(), TEXT("tutorials"));
+	bDesireResettingTutorialSeenFlagOnLoad = FParse::Param(FCommandLine::Get(), TEXT("ResetTutorials"));
 }
 
-/** Generates an analytics name for the given tutorial path string */
-FString FIntroTutorials::AnalyticsEventNameFromTutorialPath(const FString& TutorialPath) const
+FString FIntroTutorials::AnalyticsEventNameFromTutorial(const FString& BaseEventName, UEditorTutorial* Tutorial)
 {
-	// strip off everything but the last folder, supporting both types of slashes
+	FString TutorialPath = Tutorial->GetOutermost()->GetFName().ToString();
+
+	// strip off everything but the asset name
 	FString RightStr;
-	bool bSplit = TutorialPath.Split( TEXT("/"), NULL, &RightStr, ESearchCase::IgnoreCase, ESearchDir::FromEnd );
-	if (!bSplit)
-	{
-		TutorialPath.Split( TEXT("\\"), NULL, &RightStr, ESearchCase::IgnoreCase, ESearchDir::FromEnd );
-	}
+	TutorialPath.Split( TEXT("/"), NULL, &RightStr, ESearchCase::IgnoreCase, ESearchDir::FromEnd );
 
 	// then append that to the header
 	// e.g. Rocket.Tutorials.ClosedInEditorTutorial
-	FString OutStr(TEXT("Rocket.Tutorials.Closed"));
+	FString OutStr = BaseEventName;
 	OutStr += RightStr;
 
 	return OutStr;
@@ -104,8 +81,6 @@ TSharedRef<FExtender> FIntroTutorials::AddSummonBlueprintTutorialsMenuExtender(c
 
 void FIntroTutorials::StartupModule()
 {
-	GConfig->GetBool(*IntroTutorialConfigSection, *DisableTutorialsSettingName, bDisableTutorials, GEditorGameAgnosticIni);
-
 	// This code can run with content commandlets. Slate is not initialized with commandlets and the below code will fail.
 	if (!bDisableTutorials && !IsRunningCommandlet())
 	{
@@ -123,11 +98,6 @@ void FIntroTutorials::StartupModule()
 		// Add menu option to blueprint editor as well
 		FBlueprintEditorModule& BPEditorModule = FModuleManager::LoadModuleChecked<FBlueprintEditorModule>( "Kismet" );
 		BPEditorModule.GetMenuExtensibilityManager()->GetExtenderDelegates().Add(FAssetEditorExtender::CreateRaw(this, &FIntroTutorials::AddSummonBlueprintTutorialsMenuExtender));
-		
-		// Add hook for when specific asset editor is opened
-		FAssetEditorManager::Get().OnAssetEditorOpened().AddRaw(this, &FIntroTutorials::OnAssetEditorOpened);
-
-		IMatineeModule::Get().OnMatineeEditorOpened().AddRaw(this, &FIntroTutorials::OnMatineeEditorOpened);
 
 		// Add hook for when AddToCodeProject dialog window is opened
 		FGameProjectGenerationModule::Get().OnAddCodeToProjectDialogOpened().AddRaw(this, &FIntroTutorials::OnAddCodeToProjectDialogOpened);
@@ -135,72 +105,12 @@ void FIntroTutorials::StartupModule()
 		// Add hook for New Project dialog window is opened
 		//FGameProjectGenerationModule::Get().OnNewProjectProjectDialogOpened().AddRaw(this, &FIntroTutorials::OnNewProjectDialogOpened);
 
-
-		// Add hook for when editor changes modes (e.g. Place/Paint/Landscape/Foliage)
-		GLevelEditorModeTools().OnEditorModeChanged().AddRaw(this, &FIntroTutorials::OnEditorModeChanged);
-
 		FSourceCodeNavigation::AccessOnCompilerNotFound().AddRaw( this, &FIntroTutorials::HandleCompilerNotFound );
-
-		// set up some tutorial data
-
-		// NOTE: we match classes based on an "is-a" relationship, so place more derived classes in the map before base classes so they are encountered first
-		AssetEditorTutorialPropertyMap.Add(UTexture::StaticClass(), FWelcomeTutorialProperties(TEXT("Shared/Tutorials/InTextureEditorTutorial"), TEXT("SeenTextureEditorWelcome"), FString("6CC90D96-03D6-4847-B1DF-EB3018C36C4E")));
-		AssetEditorTutorialPropertyMap.Add(UMaterial::StaticClass(), FWelcomeTutorialProperties(TEXT("Shared/Tutorials/InMaterialEditorTutorial"), TEXT("SeenMaterialEditorWelcome"), FString("2290389E-A428-46C1-B9BD-A7F110CDF83E")));
-		AssetEditorTutorialPropertyMap.Add(UAnimBlueprint::StaticClass(), FWelcomeTutorialProperties(TEXT("Shared/Tutorials/InPersonaAnimBlueprintEditorTutorial"), TEXT("SeenPersonaAnimBlueprintEditorWelcome"), FString("FE5CD131-E4AF-4FE6-9269-12B573B66CA8")));
-		AssetEditorTutorialPropertyMap.Add(UBlueprint::StaticClass(), FWelcomeTutorialProperties(FWelcomeTutorialProperties::FWelcomeTutorialPropertiesChooserDelegate::CreateRaw(this, &FIntroTutorials::ChooseBlueprintWelcomeTutorial)));
-		AssetEditorTutorialPropertyMap.Add(UStaticMesh::StaticClass(), FWelcomeTutorialProperties(TEXT("Shared/Tutorials/InStaticMeshEditorTutorial"), TEXT("SeenStaticMeshEditorWelcome"), FString("762FCAD0-9384-4A6B-8D57-92DA771AD890")));
-		AssetEditorTutorialPropertyMap.Add(UDestructibleMesh::StaticClass(), FWelcomeTutorialProperties(TEXT("Shared/Tutorials/InDestructibleMeshEditorTutorial"), TEXT("SeenDestructibleMeshEditorWelcome"), FString("415C385F-C7D3-4CE8-BCC2-F19BB512AF06")));
-		AssetEditorTutorialPropertyMap.Add(USkeletalMesh::StaticClass(), FWelcomeTutorialProperties(TEXT("Shared/Tutorials/InSkeletalMeshEditorTutorial"), TEXT("SeenSkeletalMeshEditorWelcome"), FString("F8BF69A0-391F-4612-B289-B2B2B3F7F428")));
-		AssetEditorTutorialPropertyMap.Add(UParticleSystem::StaticClass(), FWelcomeTutorialProperties(TEXT("Shared/Tutorials/InParticleSystemEditorTutorial"), TEXT("SeenParticleSystemEditorWelcome"), FString("B5ACB9D0-229B-4DA2-A3D6-001A368A48B1")));
-		AssetEditorTutorialPropertyMap.Add(USoundCue::StaticClass(), FWelcomeTutorialProperties(TEXT("Shared/Tutorials/InSoundCueEditorTutorial"), TEXT("SeenSoundCueEditorWelcome"), FString("3C12B4B3-36F8-48D2-A8B0-874CCD5D3891")));
-		AssetEditorTutorialPropertyMap.Add(UMaterialInstance::StaticClass(), FWelcomeTutorialProperties(TEXT("Shared/Tutorials/InMaterialInstanceEditorTutorial"), TEXT("SeenMaterialInstanceEditorWelcome"), FString("A1AAC488-8389-4B71-925D-7A2CFB65F8D4")));
-		AssetEditorTutorialPropertyMap.Add(UPhysicsAsset::StaticClass(), FWelcomeTutorialProperties(TEXT("Shared/Tutorials/InPhATEditorTutorial"), TEXT("SeenPhATEditorWelcome"), FString("B576D8AB-9248-4FE4-80F7-96EDA87BBAD3")));
-
-		// note these 3 go to the same editor
-		AssetEditorTutorialPropertyMap.Add(UAnimMontage::StaticClass(), FWelcomeTutorialProperties(TEXT("Shared/Tutorials/InPersonaAnimEditorTutorial"), TEXT("SeenPersonaAnimEditorWelcome"), FString("8C7C3772-9135-426C-93A7-5937C5CBEA7B")));
-		AssetEditorTutorialPropertyMap.Add(UAnimSequence::StaticClass(), FWelcomeTutorialProperties(TEXT("Shared/Tutorials/InPersonaAnimEditorTutorial"), TEXT("SeenPersonaAnimEditorWelcome"), FString("8C7C3772-9135-426C-93A7-5937C5CBEA7B")));
-		AssetEditorTutorialPropertyMap.Add(UBlendSpace::StaticClass(), FWelcomeTutorialProperties(TEXT("Shared/Tutorials/InPersonaAnimEditorTutorial"), TEXT("SeenPersonaAnimEditorWelcome"), FString("8C7C3772-9135-426C-93A7-5937C5CBEA7B")));
-
-		// map editor modes to tutorial properties
-		//@todo Placement mode appears to be firing erroneously when switching to other modes, want to understand why before enabling this mode
-		//EditorModeTutorialPropertyMap.Add(FBuiltinEditorModes::EM_Placement, FWelcomeTutorialProperties(TEXT("Shared/Tutorials/EditorPlacementModeTutorial"), TEXT("SeenEditorPlacementModeWelcome")));
-		EditorModeTutorialPropertyMap.Add(FBuiltinEditorModes::EM_Landscape, FWelcomeTutorialProperties(TEXT("Shared/Tutorials/EditorLandscapeModeTutorial"), TEXT("SeenEditorLandscapeModeWelcome"), FString("8B33592A-CCEE-4129-AE3F-77A5B68955CB")));
-		EditorModeTutorialPropertyMap.Add(FBuiltinEditorModes::EM_MeshPaint, FWelcomeTutorialProperties(TEXT("Shared/Tutorials/EditorMeshPaintModeTutorial"), TEXT("SeenEditorMeshPaintModeWelcome"), FString("8E36B45D-D249-42E1-ABEE-2ABCBCE8ED26")));
-		EditorModeTutorialPropertyMap.Add(FBuiltinEditorModes::EM_Foliage, FWelcomeTutorialProperties(TEXT("Shared/Tutorials/EditorFoliageModeTutorial"), TEXT("SeenEditorFoliageModeWelcome"), FString("9DD00A25-16D5-40D5-BDC4-98DF70CEB0F7")));
-
-		// init survey map
-		{
-			for (auto It = AssetEditorTutorialPropertyMap.CreateConstIterator(); It; ++It)
-			{
-				const FWelcomeTutorialProperties& Props = It.Value();
-				TutorialSurveyMap.Add(Props.TutorialPath, Props.SurveyGuid);
-			}
-			for (auto It = EditorModeTutorialPropertyMap.CreateConstIterator(); It; ++It)
-			{
-				const FWelcomeTutorialProperties& Props = It.Value();
-				TutorialSurveyMap.Add(Props.TutorialPath, Props.SurveyGuid);
-			}
-			TutorialSurveyMap.Add(ClassBlueprintWelcomeTutorial.TutorialPath, ClassBlueprintWelcomeTutorial.SurveyGuid);
-			TutorialSurveyMap.Add(MacroLibraryBlueprintWelcomeTutorial.TutorialPath, MacroLibraryBlueprintWelcomeTutorial.SurveyGuid);
-			TutorialSurveyMap.Add(InterfaceBlueprintWelcomeTutorial.TutorialPath, InterfaceBlueprintWelcomeTutorial.SurveyGuid);
-			TutorialSurveyMap.Add(LevelScriptBlueprintWelcomeTutorial.TutorialPath, LevelScriptBlueprintWelcomeTutorial.SurveyGuid);
-			TutorialSurveyMap.Add(AddCodeToProjectWelcomeTutorial.TutorialPath, AddCodeToProjectWelcomeTutorial.SurveyGuid);
-			TutorialSurveyMap.Add(MatineeEditorWelcomeTutorial.TutorialPath, MatineeEditorWelcomeTutorial.SurveyGuid);
-
-			FGuid SurveyGuid;
-
-			// set up some paths for surveys and analytics
-			SurveyGuid.ParseExact(FString("B6C19629-3172-4924-91FE-AC557F66703C"), EGuidFormats::DigitsWithHyphens, SurveyGuid);
-			TutorialSurveyMap.Add(InEditorTutorialPath, SurveyGuid);
-
-			SurveyGuid.ParseExact(FString("3ADE38A2-9FCA-4BA1-B0A3-F64DDFAB2A0E"), EGuidFormats::DigitsWithHyphens, SurveyGuid);
-			TutorialSurveyMap.Add(InEditorGamifiedTutorialPath, SurveyGuid);
-		}
-
+	
 		// maybe reset all the "have I seen this once" flags
 		if (bDesireResettingTutorialSeenFlagOnLoad)
 		{
-			ResetWelcomeTutorials();
+			GetMutableDefault<UTutorialStateSettings>()->ClearProgress();
 		}
 	}
 
@@ -239,8 +149,6 @@ void FIntroTutorials::ShutdownModule()
 	if (!bDisableTutorials && !IsRunningCommandlet())
 	{
 		FSourceCodeNavigation::AccessOnCompilerNotFound().RemoveAll( this );
-
-		GLevelEditorModeTools().OnEditorModeChanged().RemoveAll(this);
 	}
 
 	if (BlueprintEditorExtender.IsValid() && FModuleManager::Get().IsModuleLoaded("Kismet"))
@@ -323,12 +231,6 @@ void FIntroTutorials::MainFrameLoad(TSharedPtr<SWindow> InRootWindow, bool bIsNe
 {
 	if (!bIsNewProjectWindow)
 	{
-		// Save off pointer to root window so we can parent the tutorial window to it when summoned
-		if (InRootWindow.IsValid())
-		{
-			RootWindow = InRootWindow;
-		}
-
 		// install a root widget for the tutorial overlays to hang off
 		if(InRootWindow.IsValid() && !TutorialRoot.IsValid())
 		{
@@ -339,7 +241,7 @@ void FIntroTutorials::MainFrameLoad(TSharedPtr<SWindow> InRootWindow, bool bIsNe
 		}
 
 		// See if we should show 'welcome' screen
-		MaybeOpenWelcomeTutorial(UE4WelcomeTutorial.TutorialPath, UE4WelcomeTutorial.SeenOnceSettingName);
+		MaybeOpenWelcomeTutorial();
 	}
 }
 
@@ -351,209 +253,37 @@ void FIntroTutorials::SummonTutorialHome()
 
 void FIntroTutorials::SummonBlueprintTutorialHome(UObject* Asset, bool bForceWelcome)
 {
-	if(Asset != nullptr)
-	{
-		CurrentObjectClass = Asset->GetClass();
-	}
+	UBlueprint* BP = CastChecked<UBlueprint>(Asset);
 
-	FWelcomeTutorialProperties const* Tutorial = ChooseBlueprintWelcomeTutorial(Asset, bForceWelcome);
-	check(Tutorial);
-	SummonTutorialWindowForPage(Tutorial->TutorialPath);
-
-	// make sure we've seen this tutorial
-	GConfig->SetBool(*IntroTutorialConfigSection, *Tutorial->SeenOnceSettingName, true, GEditorGameAgnosticIni);
-}
-
-void FIntroTutorials::SummonTutorialWindowForPage(const FString& Path)
-{
-	TSharedPtr<SWindow> Window;
-
-	// If window already exists, bring it to the front
-	if (TutorialWindow.IsValid())
+	FName Context("BlueprintOverview");
+	if(!bForceWelcome)
 	{
-		Window = TutorialWindow.Pin();
-		Window->BringToFront();
+		Context = FBlueprintEditor::GetContextFromBlueprintType(BP->BlueprintType);
 	}
-	// If not, create a new window
-	else
-	{
-		// Window
-		Window = SNew(SWindow)
-			.Title(LOCTEXT( "WindowTitle", "Tutorials" ))
-			.ClientSize(FVector2D(660.f, 637.f))
-			.SupportsMaximize(false)
-			.SupportsMinimize(false);
-
-		TutorialWindow = Window;
-
-		if(RootWindow.IsValid() && !PLATFORM_MAC)
-		{
-			FSlateApplication::Get().AddWindowAsNativeChild(Window.ToSharedRef(), RootWindow.Pin().ToSharedRef());
-		}
-		else
-		{
-			FSlateApplication::Get().AddWindow(Window.ToSharedRef());
-		}
-
-		// tutorial Widget
-		TSharedRef<SIntroTutorials> TutorialWidgetRef = 
-			SNew(SIntroTutorials) 
-			.ParentWindow(Window)
-			.HomeButtonVisibility(TAttribute<EVisibility>::Create(TAttribute<EVisibility>::FGetter::CreateRaw(this, &FIntroTutorials::GetHomeButtonVisibility)))
-			.OnGotoNextTutorial(FOnGotoNextTutorial::CreateRaw(this, &FIntroTutorials::HandleGotoNextTutorial));
-
-		Window->SetContent(TutorialWidgetRef);
 
-		Window->SetOnWindowClosed( FOnWindowClosed::CreateRaw(this, &FIntroTutorials::OnTutorialWindowClosed) );
-		TutorialWidgetRef->SetOnGoHome( FOnGoHome::CreateRaw(this, &FIntroTutorials::OnTutorialDismissed) );
-
-		TutorialWidget = TutorialWidgetRef;
-
-		Window->BringToFront();
-	}
-
-	// Change page to desired path
-	if(TutorialWidget.IsValid())
-	{
-		TutorialWidget.Pin()->ChangePage(Path);
-	}
-}
-
-void FIntroTutorials::OnTutorialWindowClosed(const TSharedRef<SWindow>& Window)
-{
-	OnTutorialDismissed();
-}
-
-void FIntroTutorials::OnTutorialDismissed() const
-{
-	TSharedPtr<SIntroTutorials> WidgetPtr = TutorialWidget.Pin();
-	if (WidgetPtr.IsValid())
-	{
-		FString const CurrentPagePath = WidgetPtr->GetCurrentPagePath();
-
-		// submit analytics data
-		if( FEngineAnalytics::IsAvailable() )
-		{
-			// prepare and send analytics data
-			bool const bQuitOnWelcomeScreen = (CurrentPagePath == WelcomeTutorialPath);
-
-			FString const CurrentExcerptTitle = bQuitOnWelcomeScreen ? TEXT("Welcome Screen") : WidgetPtr->GetCurrentExcerptIdentifierName();
-			int32 const CurrentExcerptIndex = bQuitOnWelcomeScreen ? -1 : WidgetPtr->GetCurrentExcerptIndex();
-			float const CurrentPageElapsedTime = bQuitOnWelcomeScreen ? 0.f : WidgetPtr->GetCurrentPageElapsedTime();
-
-			TArray<FAnalyticsEventAttribute> EventAttributes;
-			EventAttributes.Add(FAnalyticsEventAttribute(TEXT("LastExcerptIndex"), CurrentExcerptIndex));
-			EventAttributes.Add(FAnalyticsEventAttribute(TEXT("LastExcerptTitle"), CurrentExcerptTitle));
-			EventAttributes.Add(FAnalyticsEventAttribute(TEXT("TimeSpentInTutorial"), CurrentPageElapsedTime));
-
-			FString AnalyticsEventName = AnalyticsEventNameFromTutorialPath(CurrentPagePath);
-			
-			FEngineAnalytics::GetProvider().RecordEvent( AnalyticsEventName, EventAttributes );
-		}
-
-		// kick off survey, but not if it was dismissed immediately
-		if ( bEnablePostTutorialSurveys && IEpicSurveyModule::IsAvailable() && (WidgetPtr->GetCurrentExcerptIndex() > 0) )
-		{
-			FGuid const* const SurveyID = TutorialSurveyMap.Find(CurrentPagePath);
-			if (SurveyID && !SurveyID->IsValid())
-			{
-				// launch end-of-tutorial survey if one is desired
-				IEpicSurveyModule::Get().PromptSurvey(*SurveyID);
-			}
-		}
-	}	
-}
-
-void FIntroTutorials::OnAssetEditorOpened(UObject* Asset)
-{
-	if(Asset != nullptr)
-	{
-		FWelcomeTutorialProperties const* const FoundProps = FindAssetEditorTutorialProperties(Asset->GetClass());
-		if (FoundProps)
-		{
-			// run delegate if it exists
-			FWelcomeTutorialProperties const* const PropsToUse = FoundProps->ChooserDelegate.IsBound()
-				? FoundProps->ChooserDelegate.Execute(Asset)
-				: FoundProps;
-
-			if (PropsToUse)
-			{
-				if(MaybeOpenWelcomeTutorial(PropsToUse->TutorialPath, PropsToUse->SeenOnceSettingName))
-				{
-					CurrentObjectClass = Asset->GetClass();
-				}
-			}
-		}
-	}
-}
-
-
-FWelcomeTutorialProperties const* FIntroTutorials::ChooseBlueprintWelcomeTutorial(UObject* BlueprintObject)
-{
-	return ChooseBlueprintWelcomeTutorial(BlueprintObject, true);
-}
+	UEditorTutorial* AttractTutorial = nullptr;
+	UEditorTutorial* LaunchTutorial = nullptr;
+	FString BrowserFilter;
+	GetDefault<UEditorTutorialSettings>()->FindTutorialInfoForContext(Context, AttractTutorial, LaunchTutorial, BrowserFilter);
 
-FWelcomeTutorialProperties const* FIntroTutorials::ChooseBlueprintWelcomeTutorial(UObject* BlueprintObject, bool bForceWelcome)
-{
-	TutorialChainMap.Empty();
+	FIntroTutorials& IntroTutorials = FModuleManager::GetModuleChecked<FIntroTutorials>(TEXT("IntroTutorials"));
 
-	UBlueprint* BP = Cast<UBlueprint>(BlueprintObject);
-	if (BP)
+	if (LaunchTutorial != nullptr)
 	{
-		bool bSeenWelcome = false;
-		GConfig->GetBool(*IntroTutorialConfigSection, *BlueprintHomeTutorial.SeenOnceSettingName, bSeenWelcome, GEditorGameAgnosticIni);
-		if(!bSeenWelcome || bForceWelcome)
+		TSharedPtr<SWindow> ContextWindow;
+		TSharedPtr< IToolkit > Toolkit = FToolkitManager::Get().FindEditorForAsset(Asset);
+		if(ensure(Toolkit.IsValid()))
 		{
-			switch (BP->BlueprintType)
-			{
-			case BPTYPE_Normal:
-				TutorialChainMap.Add(BlueprintHomePath, ClassBlueprintWelcomeTutorial.TutorialPath);
-				break;
-			case BPTYPE_MacroLibrary:
-				TutorialChainMap.Add(BlueprintHomePath, MacroLibraryBlueprintWelcomeTutorial.TutorialPath);
-				break;
-			case BPTYPE_Interface:
-				TutorialChainMap.Add(BlueprintHomePath, InterfaceBlueprintWelcomeTutorial.TutorialPath);
-				break;
-			case BPTYPE_LevelScript:
-				TutorialChainMap.Add(BlueprintHomePath, LevelScriptBlueprintWelcomeTutorial.TutorialPath);
-				break;
-			}
+			ContextWindow = FSlateApplication::Get().FindWidgetWindow(Toolkit->GetToolkitHost()->GetParentWidget());
+			check(ContextWindow.IsValid());
 		}
-		else
-		{
-			switch (BP->BlueprintType)
-			{
-			case BPTYPE_Normal:
-				return &ClassBlueprintWelcomeTutorial;
-			case BPTYPE_MacroLibrary:
-				return &MacroLibraryBlueprintWelcomeTutorial;
-			case BPTYPE_Interface:
-				return &InterfaceBlueprintWelcomeTutorial;
-			case BPTYPE_LevelScript:
-				return &LevelScriptBlueprintWelcomeTutorial;
-			}
-		}
-	}
 
-	return &BlueprintHomeTutorial;
-}
-
-FWelcomeTutorialProperties const* FIntroTutorials::FindAssetEditorTutorialProperties(UClass const* Class) const
-{
-	for (auto It = AssetEditorTutorialPropertyMap.CreateConstIterator(); It; ++It)
-	{
-		UClass* const KeyClass = It.Key();
-		if (Class->IsChildOf(KeyClass))
-		{
-			return &It.Value();
-		}
+		const bool bRestart = true;
+		IntroTutorials.LaunchTutorial(LaunchTutorial, bRestart, ContextWindow);
 	}
-
-	return NULL;
 }
 
-bool FIntroTutorials::MaybeOpenWelcomeTutorial(const FString& TutorialPath, const FString& ConfigSettingName)
+bool FIntroTutorials::MaybeOpenWelcomeTutorial()
 {
 	// try editor startup tutorial
 	TSubclassOf<UEditorTutorial> EditorStartupTutorialClass = LoadClass<UEditorTutorial>(NULL, *GetDefault<UEditorTutorialSettings>()->StartupTutorial.AssetLongPathname, NULL, LOAD_None, NULL);
@@ -582,99 +312,32 @@ bool FIntroTutorials::MaybeOpenWelcomeTutorial(const FString& TutorialPath, cons
 	return false;
 }
 
-bool FIntroTutorials::MaybeOpenWelcomeTutorial(const FWelcomeTutorialProperties& TutorialProperties)
-{
-	return MaybeOpenWelcomeTutorial(TutorialProperties.TutorialPath, TutorialProperties.SeenOnceSettingName);
-}
-
-template< typename KeyType >
-void FIntroTutorials::ResetTutorialPropertyMap(TMap<KeyType, FWelcomeTutorialProperties> Map) const
-{
-	for (auto It = Map.CreateConstIterator(); It; ++It)
-	{
-		ResetTutorial(It.Value());
-	}
-}
-
-void FIntroTutorials::ResetWelcomeTutorials() const
-{
-	ResetTutorialPropertyMap(AssetEditorTutorialPropertyMap);
-	ResetTutorialPropertyMap(EditorModeTutorialPropertyMap);
-
-	ResetTutorial(UE4WelcomeTutorial);
-	ResetTutorial(BlueprintHomeTutorial);
-	ResetTutorial(ClassBlueprintWelcomeTutorial);
-	ResetTutorial(MacroLibraryBlueprintWelcomeTutorial);
-	ResetTutorial(InterfaceBlueprintWelcomeTutorial);
-	ResetTutorial(LevelScriptBlueprintWelcomeTutorial);
-	ResetTutorial(AddCodeToProjectWelcomeTutorial);
-	ResetTutorial(MatineeEditorWelcomeTutorial);
-	ResetTutorial(TemplateOverview);
-}
-
 void FIntroTutorials::OnAddCodeToProjectDialogOpened()
 {
-	MaybeOpenWelcomeTutorial(AddCodeToProjectWelcomeTutorial);
+	// @todo: add code to project dialog tutorial here?
+//	MaybeOpenWelcomeTutorial(AddCodeToProjectWelcomeTutorial);
 }
 
 void FIntroTutorials::OnNewProjectDialogOpened()
 {
-	MaybeOpenWelcomeTutorial(TemplateOverview);
-}
-
-void FIntroTutorials::OnMatineeEditorOpened()
-{
-	MaybeOpenWelcomeTutorial(MatineeEditorWelcomeTutorial);
-}
-
-void FIntroTutorials::OnEditorModeChanged(FEdMode* Mode, bool bEnteringMode)
-{
-	// editor-mode changed events can fire before the MainFrameLoad event, so wait until *after* the welcome tutorial to start showing editor-mode tutorials
-	// this avoids erroneously suppressing a mode tutorial because it opened but was then immediately replaced by the UE4 welcome tutorial
-	if (!HasSeenTutorial(UE4WelcomeTutorial))
-		return;
-	
-	// only show tutorials when entering an editor mode, not when leaving it
-	if (bEnteringMode)
-	{
-		FWelcomeTutorialProperties const* const FoundProp = EditorModeTutorialPropertyMap.Find(Mode->GetID());
-		if (FoundProp)
-		{
-			MaybeOpenWelcomeTutorial(FoundProp->TutorialPath, FoundProp->SeenOnceSettingName);
-		}
-	}
+	// @todo: new project dialog tutorial here?
+//	MaybeOpenWelcomeTutorial(TemplateOverview);
 }
 
 void FIntroTutorials::HandleCompilerNotFound()
 {
 #if PLATFORM_WINDOWS
-	SummonTutorialWindowForPage( TEXT( "Shared/Tutorials/InstallingVisualStudioTutorial" ) );
+	LaunchTutorial( TEXT( "Engine/Tutorial/Installation/InstallingVisualStudioTutorial.InstallingVisualStudioTutorial" ) );
 #elif PLATFORM_MAC
-	SummonTutorialWindowForPage( TEXT( "Shared/Tutorials/InstallingXCodeTutorial" ) );
+	LaunchTutorial( TEXT( "Engine/Tutorial/Installation/InstallingXCodeTutorial.InstallingXCodeTutorial" ) );
 #else
 	STUBBED("FIntroTutorials::HandleCompilerNotFound");
 #endif
 }
 
-void FIntroTutorials::HandleSDKNotInstalled(const FString& PlatformName, const FString& DocLink)
+void FIntroTutorials::HandleSDKNotInstalled(const FString& PlatformName, const FString& InTutorialAsset)
 {
-	SummonTutorialWindowForPage(DocLink);
-}
-
-bool FIntroTutorials::HasSeenTutorial( const FWelcomeTutorialProperties& TutProps ) const
-{
-	bool bSeenTutorial = false;
-	GConfig->GetBool(*IntroTutorialConfigSection, *TutProps.SeenOnceSettingName, bSeenTutorial, GEditorGameAgnosticIni);
-	return bSeenTutorial;
-}
-
-void FIntroTutorials::ResetTutorial(const FWelcomeTutorialProperties& TutProps) const
-{
-	FString SettingName = TutProps.SeenOnceSettingName;
-	if (!SettingName.IsEmpty())
-	{
-		GConfig->SetBool(*IntroTutorialConfigSection, *SettingName, false, GEditorGameAgnosticIni);
-	}
+	LaunchTutorial( InTutorialAsset );
 }
 
 EVisibility FIntroTutorials::GetHomeButtonVisibility() const
@@ -687,38 +350,6 @@ EVisibility FIntroTutorials::GetHomeButtonVisibility() const
 	return EVisibility::Visible;
 }
 
-FString FIntroTutorials::HandleGotoNextTutorial(const FString& InCurrentPage) const
-{
-	const FString* NextTutorial = TutorialChainMap.Find(InCurrentPage);
-	if(NextTutorial != nullptr)
-	{
-		return *NextTutorial;
-	}
-
-	return FString();
-}
-
-void FIntroTutorials::RegisterTutorialForAssetEditor(UClass* AssetClass, const FString& TutorialDocPath, const FString& TutorialHasBeenSeenSettingName, const FString& SurveyGUIDString)
-{
-	FWelcomeTutorialProperties TutorialData(TutorialDocPath, TutorialHasBeenSeenSettingName, SurveyGUIDString);
-
-	if (bDesireResettingTutorialSeenFlagOnLoad)
-	{
-		ResetTutorial(TutorialData);
-	}
-
-	TutorialSurveyMap.Add(TutorialDocPath, TutorialData.SurveyGuid);
-
-	AssetEditorTutorialPropertyMap.Add(AssetClass, TutorialData);
-}
-
-void FIntroTutorials::UnregisterTutorialForAssetEditor(UClass* AssetClass)
-{
-	const FWelcomeTutorialProperties& TutorialData = AssetEditorTutorialPropertyMap.FindChecked(AssetClass);
-	TutorialSurveyMap.Remove(TutorialData.TutorialPath);
-	AssetEditorTutorialPropertyMap.Remove(AssetClass);
-}
-
 FOnIsPicking& FIntroTutorials::OnIsPicking()
 {
 	return OnIsPickingDelegate;
@@ -732,10 +363,25 @@ void FIntroTutorials::SummonTutorialBrowser(TSharedRef<SWindow> InWindow, const
 	}
 }
 
+void FIntroTutorials::LaunchTutorial(const FString& InAssetPath, bool bInRestart, TWeakPtr<SWindow> InNavigationWindow, FSimpleDelegate OnTutorialClosed, FSimpleDelegate OnTutorialExited)
+{
+	UBlueprint* Blueprint = LoadObject<UBlueprint>(nullptr, *InAssetPath);
+	if (Blueprint && Blueprint->GeneratedClass)
+	{
+		LaunchTutorial(Blueprint->GeneratedClass->GetDefaultObject<UEditorTutorial>(), bInRestart, InNavigationWindow, OnTutorialClosed, OnTutorialExited);
+	}
+}
+
 void FIntroTutorials::LaunchTutorial(UEditorTutorial* InTutorial, bool bInRestart, TWeakPtr<SWindow> InNavigationWindow, FSimpleDelegate OnTutorialClosed, FSimpleDelegate OnTutorialExited)
 {
 	if(TutorialRoot.IsValid())
 	{
+		if(!InNavigationWindow.IsValid())
+		{
+			IMainFrameModule& MainFrameModule = FModuleManager::LoadModuleChecked<IMainFrameModule>(TEXT("MainFrame"));
+			InNavigationWindow = MainFrameModule.GetParentWindow();
+		}
+
 		TutorialRoot->LaunchTutorial(InTutorial, bInRestart, InNavigationWindow, OnTutorialClosed, OnTutorialExited);
 	}
 }
diff --git a/Engine/Source/Editor/IntroTutorials/Private/IntroTutorials.h b/Engine/Source/Editor/IntroTutorials/Private/IntroTutorials.h
--- a/Engine/Source/Editor/IntroTutorials/Private/IntroTutorials.h
+++ b/Engine/Source/Editor/IntroTutorials/Private/IntroTutorials.h
@@ -9,60 +9,9 @@ class UEditorTutorial;
 struct FTutorialContent;
 class STutorialRoot;
 
-/** Contains properties defining a "welcome" tutorial, which is auto-shown once.*/
-struct FWelcomeTutorialProperties
-{
-	/** If set, call this delegate to return the proper properties struct. */
-	DECLARE_DELEGATE_RetVal_OneParam(FWelcomeTutorialProperties const*, FWelcomeTutorialPropertiesChooserDelegate, UObject*);
-	FWelcomeTutorialPropertiesChooserDelegate ChooserDelegate;
-
-	/** Location of the tutorial doc files */
-	FString TutorialPath;
-
-	/** Ini setting name for the "have I seen this welcome screen" variable */
-	FString SeenOnceSettingName;
-
-	FGuid SurveyGuid;
-
-	FWelcomeTutorialProperties(FString InTutorialPath, FString InSeenOnceSettingName)
-		: TutorialPath(InTutorialPath), SeenOnceSettingName(InSeenOnceSettingName)
-	{}
-
-	FWelcomeTutorialProperties(FString InTutorialPath, FString InSeenOnceSettingName, FString SurveyGUIDString)
-		: TutorialPath(InTutorialPath), SeenOnceSettingName(InSeenOnceSettingName)
-	{
-		FGuid::ParseExact(SurveyGUIDString, EGuidFormats::DigitsWithHyphens, SurveyGuid);
-	}
-
-	FWelcomeTutorialProperties(FWelcomeTutorialPropertiesChooserDelegate InChooserDelegate)
-		: ChooserDelegate(InChooserDelegate)
-	{}
-
-	FWelcomeTutorialProperties() {}
-};
-
 class FIntroTutorials : public IIntroTutorials
 {
 public:
-	static const FString IntroTutorialConfigSection;
-	static const FString DisableTutorialsSettingName;
-	static const FString InEditorTutorialPath;
-	static const FString WelcomeTutorialPath;
-	static const FString InEditorGamifiedTutorialPath;
-	static const FString HomePath;
-	static const FString BlueprintHomePath;
-	static const FString TemplateOverviewPath;
-
-	static const FWelcomeTutorialProperties UE4WelcomeTutorial;
-	static const FWelcomeTutorialProperties BlueprintHomeTutorial;
-	static const FWelcomeTutorialProperties ClassBlueprintWelcomeTutorial;
-	static const FWelcomeTutorialProperties MacroLibraryBlueprintWelcomeTutorial;
-	static const FWelcomeTutorialProperties InterfaceBlueprintWelcomeTutorial;
-	static const FWelcomeTutorialProperties LevelScriptBlueprintWelcomeTutorial;
-	static const FWelcomeTutorialProperties AddCodeToProjectWelcomeTutorial;
-	static const FWelcomeTutorialProperties MatineeEditorWelcomeTutorial;
-	static const FWelcomeTutorialProperties TemplateOverview;
-
 	// ctor
 	FIntroTutorials();
 
@@ -77,6 +26,8 @@ public:
 
 	void SummonTutorialBrowser(TSharedRef<SWindow> InWindow, const FString& InFilter = TEXT(""));
 
+	static FString AnalyticsEventNameFromTutorial(const FString& BaseEventName, UEditorTutorial* Tutorial);
+
 private:
 
 	/** IModuleInterface implementation */
@@ -95,9 +46,6 @@ private:
 	/** Event to be called when the main frame is loaded */
 	void MainFrameLoad(TSharedPtr<SWindow> InRootWindow, bool bIsNewProjectWindow);
 
-	/** Summon a tutorial with the supplied path */
-	void SummonTutorialWindowForPage(const FString& Path);
-
 	/** Summon tutorial home page to front */
 	void SummonTutorialHome();
 
@@ -116,41 +64,25 @@ private:
 	/** Events to call when editor changes state in various ways */
 	void OnAddCodeToProjectDialogOpened();
 	void OnNewProjectDialogOpened();
-	void OnMatineeEditorOpened();
-	void OnEditorModeChanged(FEdMode* Mode, bool bEnteringMode);
 
 	/** Events to call when opening the compiler fails */
 	void HandleCompilerNotFound();
 
 	/** Events to call when SDK isn't installed */
-	void HandleSDKNotInstalled(const FString& PlatformName, const FString& DocLink);
+	void HandleSDKNotInstalled(const FString& PlatformName, const FString& InTutorialAsset);
 
-	bool MaybeOpenWelcomeTutorial(const FWelcomeTutorialProperties& TutorialProperties);
-	bool MaybeOpenWelcomeTutorial(const FString& TutorialPath, const FString& ConfigSettingName);
+	bool MaybeOpenWelcomeTutorial();
 	void ResetWelcomeTutorials() const;
 
-	template< typename KeyType >
-	void ResetTutorialPropertyMap(TMap<KeyType, FWelcomeTutorialProperties> Map) const;
-	void ResetTutorial(const FWelcomeTutorialProperties& TutProps) const;
-	bool HasSeenTutorial(const FWelcomeTutorialProperties& TutProps) const;
-
-	FWelcomeTutorialProperties const* ChooseBlueprintWelcomeTutorial(UObject* BlueprintObject);
-	FWelcomeTutorialProperties const* ChooseBlueprintWelcomeTutorial(UObject* BlueprintObject, bool bForceWelcome);
-	FWelcomeTutorialProperties const* FindAssetEditorTutorialProperties(UClass const* Class) const;
-
-	FString AnalyticsEventNameFromTutorialPath(const FString& TutorialPath) const;
-
 	/** Delegate for home button visibility */
 	EVisibility GetHomeButtonVisibility() const;
 
-	/** Handle linking between tutorials */
-	FString HandleGotoNextTutorial(const FString& InCurrentPage) const;
+	/** Internal helper to launch a tutorial from a path */
+	void LaunchTutorial(const FString& InAssetPath, bool bRestart = true, TWeakPtr<SWindow> InNavigationWindow = nullptr, FSimpleDelegate OnTutorialClosed = FSimpleDelegate(), FSimpleDelegate OnTutorialExited = FSimpleDelegate());
 
 public:
 
 	// IIntroTutorials interface
-	virtual void RegisterTutorialForAssetEditor(UClass* AssetClass, const FString& TutorialDocPath, const FString& TutorialHasBeenSeenSettingName, const FString& SurveyGUIDString) override;
-	virtual void UnregisterTutorialForAssetEditor(UClass* AssetClass) override;
 	virtual void LaunchTutorial(UEditorTutorial* Tutorial, bool bRestart = true, TWeakPtr<SWindow> InNavigationWindow = nullptr, FSimpleDelegate OnTutorialClosed = FSimpleDelegate(), FSimpleDelegate OnTutorialExited = FSimpleDelegate()) override;
 	virtual void CloseAllTutorialContent() override;
 	virtual TSharedRef<SWidget> CreateTutorialsWidget(FName InContext, TWeakPtr<SWindow> InContextWindow) const override;
@@ -165,32 +97,12 @@ private:
 	/** The extender to pass to the blueprint editor to extend it's window menu */
 	TSharedPtr<FExtender> BlueprintEditorExtender;
 
-	/** The tutorial window that we use to display the tutorials */
-	TWeakPtr<SWindow> TutorialWindow;
-
-	/** Widget used to display tutorial */
-	TWeakPtr<class SIntroTutorials> TutorialWidget;
-
-	/** The root window that we will always try to parent the tutorial window to */
-	TWeakPtr<SWindow> RootWindow;
-
 	/** Whether tutorials are disabled altogether */
 	bool bDisableTutorials;
 
-	TMap<FString, FGuid> TutorialSurveyMap;
-
-	TMap<UClass*, FWelcomeTutorialProperties> AssetEditorTutorialPropertyMap;
-
-	TMap<FName, FWelcomeTutorialProperties> EditorModeTutorialPropertyMap;
-
 	/** The current object we are using as a basis for displaying a tutorial */
 	UClass* CurrentObjectClass;
 
-	bool bEnablePostTutorialSurveys;
-
-	/** Map of what tutorial to go to at the end of another */
-	TMap<FString, FString> TutorialChainMap;
-
 	/** should we be clearing the 'have seen this tutorial' flag? (controlled by -tutorials on the command line) */
 	bool bDesireResettingTutorialSeenFlagOnLoad;
 
diff --git a/Engine/Source/Editor/IntroTutorials/Private/SEditorTutorials.cpp b/Engine/Source/Editor/IntroTutorials/Private/SEditorTutorials.cpp
--- a/Engine/Source/Editor/IntroTutorials/Private/SEditorTutorials.cpp
+++ b/Engine/Source/Editor/IntroTutorials/Private/SEditorTutorials.cpp
@@ -17,13 +17,14 @@ void SEditorTutorials::Construct(const FArguments& InArgs)
 	OnNextClicked = InArgs._OnNextClicked;
 	OnBackClicked = InArgs._OnBackClicked;
 	OnHomeClicked = InArgs._OnHomeClicked;
+	OnCloseClicked = InArgs._OnCloseClicked;
 	OnGetCurrentTutorial = InArgs._OnGetCurrentTutorial;
 	OnGetCurrentTutorialStage = InArgs._OnGetCurrentTutorialStage;
 
 	TutorialHome = SNew(STutorialsBrowser)
 		.Visibility(this, &SEditorTutorials::GetBrowserVisibility)
-		.OnLaunchTutorial(InArgs._OnLaunchTutorial)
 		.OnClosed(FSimpleDelegate::CreateSP(this, &SEditorTutorials::HandleCloseClicked))
+		.OnLaunchTutorial(InArgs._OnLaunchTutorial)
 		.ParentWindow(InArgs._ParentWindow);
 
 	NavigationWidget = SNew(STutorialNavigation)
@@ -119,9 +120,11 @@ EVisibility SEditorTutorials::GetNavigationVisibility() const
 
 void SEditorTutorials::HandleCloseClicked()
 {
+	OnCloseClicked.ExecuteIfBound();
+	OnTutorialClosed.ExecuteIfBound();
+
 	HandleHomeClicked();
 	bBrowserVisible = false;
-	OnTutorialClosed.ExecuteIfBound();
 }
 
 void SEditorTutorials::HandleBackClicked()
diff --git a/Engine/Source/Editor/IntroTutorials/Private/SEditorTutorials.h b/Engine/Source/Editor/IntroTutorials/Private/SEditorTutorials.h
--- a/Engine/Source/Editor/IntroTutorials/Private/SEditorTutorials.h
+++ b/Engine/Source/Editor/IntroTutorials/Private/SEditorTutorials.h
@@ -26,7 +26,7 @@ class SEditorTutorials : public SCompoundWidget
 		_Visibility = EVisibility::SelfHitTestInvisible;
 	}
 
-	SLATE_ARGUMENT(FSimpleDelegate, OnClosed)
+	SLATE_ARGUMENT(FSimpleDelegate, OnCloseClicked)
 	SLATE_ARGUMENT(TWeakPtr<SWindow>, ParentWindow)
 	SLATE_ARGUMENT(FOnNextClicked, OnNextClicked)
 	SLATE_ARGUMENT(FSimpleDelegate, OnBackClicked)
@@ -119,6 +119,9 @@ private:
 	/** Delegate fired when home button is clicked */
 	FSimpleDelegate OnHomeClicked;
 
+	/** Delegate fired when close button is clicked */
+	FSimpleDelegate OnCloseClicked;
+
 	/** Delegate fired to retrive the current tutorial */
 	FOnGetCurrentTutorial OnGetCurrentTutorial;
 
diff --git a/Engine/Source/Editor/IntroTutorials/Private/STutorialButton.cpp b/Engine/Source/Editor/IntroTutorials/Private/STutorialButton.cpp
--- a/Engine/Source/Editor/IntroTutorials/Private/STutorialButton.cpp
+++ b/Engine/Source/Editor/IntroTutorials/Private/STutorialButton.cpp
@@ -23,6 +23,7 @@ void STutorialButton::Construct(const FArguments& InArgs)
 	bTutorialCompleted = false;
 	bTutorialDismissed = false;
 	bDeferTutorialOpen = true;
+	AlertStartTime = 0.0f;
 
 	PulseAnimation.AddCurve(0.0f, TutorialButtonConstants::PulseAnimationLength, ECurveEaseFunction::Linear);
 	PulseAnimation.Play();
@@ -62,6 +63,11 @@ void STutorialButton::Tick(const FGeometry& AllottedGeometry, const double InCur
 			const bool bRestart = true;
 			IntroTutorials.LaunchTutorial(AttractTutorial, bRestart, ContextWindow);
 		}
+
+		if(ShouldShowAlert())
+		{
+			AlertStartTime = FPlatformTime::Seconds();
+		}
 	}
 	bDeferTutorialOpen = false;
 }
@@ -138,6 +144,16 @@ FReply STutorialButton::HandleButtonClicked()
 
 	if(ContextWindow.IsValid())
 	{
+		if( FEngineAnalytics::IsAvailable() )
+		{
+			TArray<FAnalyticsEventAttribute> EventAttributes;
+			EventAttributes.Add(FAnalyticsEventAttribute(TEXT("Context"), Context.ToString()));
+			EventAttributes.Add(FAnalyticsEventAttribute(TEXT("TimeSinceAlertStarted"), (AlertStartTime != 0.0f && ShouldShowAlert()) ? (FPlatformTime::Seconds() - AlertStartTime) : -1.0f));
+			EventAttributes.Add(FAnalyticsEventAttribute(TEXT("LaunchedBrowser"), ShouldLaunchBrowser());
+
+			FEngineAnalytics::GetProvider().RecordEvent( TEXT("Rocket.Tutorials.ClickedContextButton"), EventAttributes );
+		}
+
 		FIntroTutorials& IntroTutorials = FModuleManager::GetModuleChecked<FIntroTutorials>(TEXT("IntroTutorials"));
 		if(ShouldLaunchBrowser())
 		{
@@ -191,9 +207,18 @@ void STutorialButton::DismissAlert()
 	UEditorTutorial* AttractTutorial = nullptr;
 	UEditorTutorial* LaunchTutorial = nullptr;
 	FString BrowserFilter;
-	GetDefault<UEditorTutorialSettings>()->FindTutorialInfoForContext(Context, AttractTutorial, LaunchTutorial,BrowserFilter);
+	GetDefault<UEditorTutorialSettings>()->FindTutorialInfoForContext(Context, AttractTutorial, LaunchTutorial, BrowserFilter);
 	if (AttractTutorial != nullptr)
 	{
+		if( FEngineAnalytics::IsAvailable() )
+		{
+			TArray<FAnalyticsEventAttribute> EventAttributes;
+			EventAttributes.Add(FAnalyticsEventAttribute(TEXT("Context"), Context.ToString()));
+			EventAttributes.Add(FAnalyticsEventAttribute(TEXT("TimeSinceAlertStarted"), (AlertStartTime != 0.0f && ShouldShowAlert()) ? (FPlatformTime::Seconds() - AlertStartTime) : -1.0f));
+
+			FEngineAnalytics::GetProvider().RecordEvent( TEXT("Rocket.Tutorials.DismissedTutorialAlert"), EventAttributes );
+		}
+
 		const bool bDismissAcrossSessions = true;
 		GetMutableDefault<UTutorialStateSettings>()->DismissTutorial(AttractTutorial, bDismissAcrossSessions);
 		GetMutableDefault<UTutorialStateSettings>()->SaveProgress();
diff --git a/Engine/Source/Editor/IntroTutorials/Private/STutorialButton.h b/Engine/Source/Editor/IntroTutorials/Private/STutorialButton.h
--- a/Engine/Source/Editor/IntroTutorials/Private/STutorialButton.h
+++ b/Engine/Source/Editor/IntroTutorials/Private/STutorialButton.h
@@ -59,4 +59,7 @@ private:
 
 	/** Animation curve for displaying pulse */
 	FCurveSequence PulseAnimation;
+
+	/** Start time we began playing the alert animation */
+	float AlertStartTime;
 };
\ No newline at end of file
diff --git a/Engine/Source/Editor/IntroTutorials/Private/STutorialContent.cpp b/Engine/Source/Editor/IntroTutorials/Private/STutorialContent.cpp
--- a/Engine/Source/Editor/IntroTutorials/Private/STutorialContent.cpp
+++ b/Engine/Source/Editor/IntroTutorials/Private/STutorialContent.cpp
@@ -5,6 +5,8 @@
 #include "STutorialEditableText.h"
 #include "IntroTutorials.h"
 #include "TutorialText.h"
+#include "EngineAnalytics.h"
+#include "Runtime/Analytics/Analytics/Public/Interfaces/IAnalyticsProvider.h"
 
 #define LOCTEXT_NAMESPACE "STutorialContent"
 
@@ -219,6 +221,17 @@ int32 STutorialContent::OnPaint( const FPaintArgs& Args, const FGeometry& Allott
 	return SCompoundWidget::OnPaint(Args, AllottedGeometry, MyClippingRect, OutDrawElements, LayerId, InWidgetStyle, bParentEnabled);
 }
 
+FReply STutorialContent::OnMouseButtonUp( const FGeometry& MyGeometry, const FPointerEvent& MouseEvent )
+{
+	if (!bIsStandalone && MouseEvent.GetEffectingButton() == EKeys::RightMouseButton)
+	{
+		FSlateApplication::Get().PushMenu(AsShared(), HandleGetMenuContent(), MouseEvent.GetScreenSpacePosition(), FPopupTransitionEffect(FPopupTransitionEffect::ContextMenu));
+		return FReply::Handled();
+	}
+	
+	return FReply::Unhandled();
+}
+
 /** Helper function to generate title widget, if any */
 static TSharedRef<SWidget> GetStageTitle(const FExcerpt& InExcerpt, int32 InCurrentExcerptIndex)
 {
@@ -461,6 +474,19 @@ TSharedRef<SWidget> STutorialContent::HandleGetMenuContent()
 			)
 		);
 
+		if(IsNextEnabled.Get())
+		{
+			MenuBuilder.AddMenuEntry(
+				LOCTEXT("NextLabel", "Next"),
+				GetNextButtonTooltip(),
+				FSlateIcon(),
+				FUIAction(
+					FExecuteAction::CreateSP(this, &STutorialContent::HandleNextSelected),
+					FCanExecuteAction()
+				)
+			);
+		}
+
 		if(IsBackEnabled.Get())
 		{
 			MenuBuilder.AddMenuEntry(
@@ -507,6 +533,11 @@ void STutorialContent::HandleExitSelected()
 	OnClosed.ExecuteIfBound();
 }
 
+void STutorialContent::HandleNextSelected()
+{
+	OnNextClicked.ExecuteIfBound();
+}
+
 void STutorialContent::HandleBackSelected()
 {
 	OnBackClicked.ExecuteIfBound();
@@ -518,12 +549,25 @@ void STutorialContent::HandleRestartSelected()
 	{
 		FIntroTutorials& IntroTutorials = FModuleManager::GetModuleChecked<FIntroTutorials>(TEXT("IntroTutorials"));
 		const bool bRestart = true;
-		IntroTutorials.LaunchTutorial(Tutorial.Get(), bRestart, FSlateApplication::Get().FindWidgetWindow(AsShared()));		
+		IntroTutorials.LaunchTutorial(Tutorial.Get(), bRestart, FSlateApplication::Get().FindWidgetWindow(AsShared()));
+
+		if( FEngineAnalytics::IsAvailable() )
+		{
+			FEngineAnalytics::GetProvider().RecordEvent( FIntroTutorials::AnalyticsEventNameFromTutorial(TEXT("Rocket.Tutorials.Restarted"), Tutorial.Get()) );
+		}
 	}
 }
 
 void STutorialContent::HandleBrowseSelected()
 {
+	if( FEngineAnalytics::IsAvailable() && Tutorial.IsValid())
+	{
+		TArray<FAnalyticsEventAttribute> EventAttributes;
+		EventAttributes.Add(FAnalyticsEventAttribute(TEXT("FromTutorial"), FIntroTutorials::AnalyticsEventNameFromTutorial(TEXT(""), Tutorial.Get())));
+
+		FEngineAnalytics::GetProvider().RecordEvent( TEXT("Rocket.Tutorials.OpenedBrowser"), EventAttributes );
+	}
+
 	OnHomeClicked.ExecuteIfBound();
 }
 
diff --git a/Engine/Source/Editor/IntroTutorials/Private/STutorialContent.h b/Engine/Source/Editor/IntroTutorials/Private/STutorialContent.h
--- a/Engine/Source/Editor/IntroTutorials/Private/STutorialContent.h
+++ b/Engine/Source/Editor/IntroTutorials/Private/STutorialContent.h
@@ -62,6 +62,7 @@ class STutorialContent : public SCompoundWidget
 
 	/** SWidget implementation */
 	virtual int32 OnPaint( const FPaintArgs& Args, const FGeometry& AllottedGeometry, const FSlateRect& MyClippingRect, FSlateWindowElementList& OutDrawElements, int32 LayerId, const FWidgetStyle& InWidgetStyle, bool bParentEnabled ) const override;
+	virtual FReply OnMouseButtonUp( const FGeometry& MyGeometry, const FPointerEvent& MouseEvent ) override;
 
 	/** Helper function to generate widgets from an FTutorialContent struct */
 	static TSharedRef<SWidget> GenerateContentWidget(const FTutorialContent& InContent, float WrapTextAt, TSharedPtr<IDocumentationPage>& OutDocumentationPage);
@@ -113,6 +114,9 @@ private:
 	/** Delegate handler for going to the previous stage of the tutorial */
 	void HandleBackSelected();
 
+	/** Delegate handler for going to the next stage of the tutorial */
+	void HandleNextSelected();
+
 	/** Delegate handler for restarting the tutorial */
 	void HandleRestartSelected();
 
diff --git a/Engine/Source/Editor/IntroTutorials/Private/STutorialOverlay.cpp b/Engine/Source/Editor/IntroTutorials/Private/STutorialOverlay.cpp
--- a/Engine/Source/Editor/IntroTutorials/Private/STutorialOverlay.cpp
+++ b/Engine/Source/Editor/IntroTutorials/Private/STutorialOverlay.cpp
@@ -68,7 +68,7 @@ void STutorialOverlay::Construct(const FArguments& InArgs, UEditorTutorial* InTu
 						.VAlign(WidgetContent.VerticalAlignment)
 						.Offset(WidgetContent.Offset)
 						.IsStandalone(bIsStandalone)
-						.OnClosed(OnClosed)
+						.OnClosed(InArgs._OnClosed)
 						.OnNextClicked(InArgs._OnNextClicked)
 						.OnHomeClicked(InArgs._OnHomeClicked)
 						.OnBackClicked(InArgs._OnBackClicked)
diff --git a/Engine/Source/Editor/IntroTutorials/Private/STutorialRoot.cpp b/Engine/Source/Editor/IntroTutorials/Private/STutorialRoot.cpp
--- a/Engine/Source/Editor/IntroTutorials/Private/STutorialRoot.cpp
+++ b/Engine/Source/Editor/IntroTutorials/Private/STutorialRoot.cpp
@@ -9,6 +9,8 @@
 #include "ToolkitManager.h"
 #include "IToolkit.h"
 #include "IToolkitHost.h"
+#include "EngineAnalytics.h"
+#include "Runtime/Analytics/Analytics/Public/Interfaces/IAnalyticsProvider.h"
 
 #define LOCTEXT_NAMESPACE "STutorialRoot"
 
@@ -16,6 +18,7 @@ void STutorialRoot::Construct(const FArguments& InArgs)
 {
 	CurrentTutorial = nullptr;
 	CurrentTutorialStage = 0;
+	CurrentTutorialStartTime = FPlatformTime::Seconds();
 
 	ChildSlot
 	[
@@ -44,6 +47,7 @@ void STutorialRoot::MaybeAddOverlay(TSharedRef<SWindow> InWindow)
 					.OnNextClicked(FOnNextClicked::CreateSP(this, &STutorialRoot::HandleNextClicked))
 					.OnBackClicked(FSimpleDelegate::CreateSP(this, &STutorialRoot::HandleBackClicked))
 					.OnHomeClicked(FSimpleDelegate::CreateSP(this, &STutorialRoot::HandleHomeClicked))
+					.OnCloseClicked(FSimpleDelegate::CreateSP(this, &STutorialRoot::HandleCloseClicked))
 					.OnGetCurrentTutorial(FOnGetCurrentTutorial::CreateSP(this, &STutorialRoot::HandleGetCurrentTutorial))
 					.OnGetCurrentTutorialStage(FOnGetCurrentTutorialStage::CreateSP(this, &STutorialRoot::HandleGetCurrentTutorialStage))
 					.OnLaunchTutorial(FOnLaunchTutorial::CreateSP(this, &STutorialRoot::LaunchTutorial))
@@ -108,6 +112,8 @@ void STutorialRoot::LaunchTutorial(UEditorTutorial* InTutorial, bool bInRestart,
 			}
 		}
 
+		CurrentTutorialStartTime = FPlatformTime::Seconds();
+
 		// launch tutorial for all windows we wrap - any tutorial can display over any window
 		for(auto& TutorialWidget : TutorialWidgets)
 		{
@@ -155,6 +161,15 @@ void STutorialRoot::HandleNextClicked(TWeakPtr<SWindow> InNavigationWindow)
 
 void STutorialRoot::HandleBackClicked()
 {
+	if( FEngineAnalytics::IsAvailable() && CurrentTutorial != nullptr)
+	{
+		TArray<FAnalyticsEventAttribute> EventAttributes;
+		EventAttributes.Add(FAnalyticsEventAttribute(TEXT("Context.Tutorial"), FIntroTutorials::AnalyticsEventNameFromTutorial(TEXT(""), CurrentTutorial)));
+		EventAttributes.Add(FAnalyticsEventAttribute(TEXT("Context.StageIndex"), CurrentTutorialStage));
+
+		FEngineAnalytics::GetProvider().RecordEvent( TEXT("Rocket.Tutorials.ClickedBackButton"), EventAttributes );
+	}
+
 	GoToPreviousStage();
 
 	for(auto& TutorialWidget : TutorialWidgets)
@@ -261,4 +276,30 @@ void STutorialRoot::GoToNextStage(TWeakPtr<SWindow> InNavigationWindow)
 	}
 }
 
+void STutorialRoot::HandleCloseClicked()
+{
+	// submit analytics data
+	if( FEngineAnalytics::IsAvailable() && CurrentTutorial != nullptr && CurrentTutorialStage < CurrentTutorial->Stages.Num() )
+	{
+		UEditorTutorial* AttractTutorial = nullptr;
+		UEditorTutorial* LaunchTutorial = nullptr;
+		FString BrowserFilter;
+		GetDefault<UEditorTutorialSettings>()->FindTutorialInfoForContext(TEXT("LevelEditor"), AttractTutorial, LaunchTutorial, BrowserFilter);
+
+		// prepare and send analytics data
+		bool const bClosedInitialAttract = (CurrentTutorial == AttractTutorial);
+
+		FString const CurrentExcerptTitle = bClosedInitialAttract ? TEXT("InitialAttract") : CurrentTutorial->Stages[CurrentTutorialStage].Name.ToString();
+		int32 const CurrentExcerptIndex = bClosedInitialAttract ? -1 : CurrentTutorialStage;
+		float const CurrentPageElapsedTime = bClosedInitialAttract ? 0.f : (float)(FPlatformTime::Seconds() - CurrentTutorialStartTime);
+
+		TArray<FAnalyticsEventAttribute> EventAttributes;
+		EventAttributes.Add(FAnalyticsEventAttribute(TEXT("LastStageIndex"), CurrentExcerptIndex));
+		EventAttributes.Add(FAnalyticsEventAttribute(TEXT("LastStageTitle"), CurrentExcerptTitle));
+		EventAttributes.Add(FAnalyticsEventAttribute(TEXT("TimeSpentInTutorial"), CurrentPageElapsedTime));
+			
+		FEngineAnalytics::GetProvider().RecordEvent( FIntroTutorials::AnalyticsEventNameFromTutorial(TEXT("Rocket.Tutorials.Closed"), CurrentTutorial), EventAttributes );
+	}
+}
+
 #undef LOCTEXT_NAMESPACE
\ No newline at end of file
diff --git a/Engine/Source/Editor/IntroTutorials/Private/STutorialRoot.h b/Engine/Source/Editor/IntroTutorials/Private/STutorialRoot.h
--- a/Engine/Source/Editor/IntroTutorials/Private/STutorialRoot.h
+++ b/Engine/Source/Editor/IntroTutorials/Private/STutorialRoot.h
@@ -46,15 +46,18 @@ public:
 	void GoToNextStage(TWeakPtr<SWindow> InNavigationWindow);
 
 private:
-	/** Handle when the next button is clicked - forward navigation to other overlays */
+	/** Handle when the next button is clicked */
 	void HandleNextClicked(TWeakPtr<SWindow> InNavigationWindow);
 
-	/** Handle when the back button is clicked - forward navigation to other overlays */
+	/** Handle when the back button is clicked */
 	void HandleBackClicked();
 
-	/** Handle when the home button is clicked - forward navigation to other overlays */
+	/** Handle when the home button is clickeds */
 	void HandleHomeClicked();
 
+	/** Handle when the close button is clicked */
+	void HandleCloseClicked();
+
 	/** Handle retrieving the current tutorial */
 	UEditorTutorial* HandleGetCurrentTutorial();
 
@@ -73,4 +76,7 @@ private:
 
 	/** Current stage of tutorial */
 	int32 CurrentTutorialStage;
+
+	/** Start time of the current tutorial, if any */
+	float CurrentTutorialStartTime;
 };
\ No newline at end of file
diff --git a/Engine/Source/Editor/IntroTutorials/Private/STutorialsBrowser.cpp b/Engine/Source/Editor/IntroTutorials/Private/STutorialsBrowser.cpp
--- a/Engine/Source/Editor/IntroTutorials/Private/STutorialsBrowser.cpp
+++ b/Engine/Source/Editor/IntroTutorials/Private/STutorialsBrowser.cpp
@@ -8,6 +8,8 @@
 #include "EditorTutorialSettings.h"
 #include "TutorialStateSettings.h"
 #include "AssetRegistryModule.h"
+#include "EngineAnalytics.h"
+#include "Runtime/Analytics/Analytics/Public/Interfaces/IAnalyticsProvider.h"
 
 #define LOCTEXT_NAMESPACE "TutorialsBrowser"
 
@@ -693,6 +695,14 @@ FSlateColor STutorialsBrowser::GetBackButtonColor() const
 
 void STutorialsBrowser::OnTutorialSelected(UEditorTutorial* InTutorial, bool bRestart)
 {
+	if( FEngineAnalytics::IsAvailable() && InTutorial != nullptr )
+	{
+		TArray<FAnalyticsEventAttribute> EventAttributes;
+		EventAttributes.Add(FAnalyticsEventAttribute(TEXT("Restarted"), bRestart));
+
+		FEngineAnalytics::GetProvider().RecordEvent( FIntroTutorials::AnalyticsEventNameFromTutorial(TEXT("Rocket.Tutorials.LaunchedFromBrowser"), InTutorial), EventAttributes );
+	}
+
 	OnLaunchTutorial.ExecuteIfBound(InTutorial, bRestart, ParentWindow, FSimpleDelegate(), FSimpleDelegate());
 }
 
diff --git a/Engine/Source/Editor/IntroTutorials/Private/TutorialStateSettings.cpp b/Engine/Source/Editor/IntroTutorials/Private/TutorialStateSettings.cpp
--- a/Engine/Source/Editor/IntroTutorials/Private/TutorialStateSettings.cpp
+++ b/Engine/Source/Editor/IntroTutorials/Private/TutorialStateSettings.cpp
@@ -118,3 +118,11 @@ void UTutorialStateSettings::SaveProgress()
 
 	SaveConfig();
 }
+
+void UTutorialStateSettings::ClearProgress()
+{
+	ProgressMap.Empty();
+	TutorialsProgress.Empty();
+
+	SaveConfig();
+}
\ No newline at end of file
diff --git a/Engine/Source/Editor/IntroTutorials/Private/TutorialStateSettings.h b/Engine/Source/Editor/IntroTutorials/Private/TutorialStateSettings.h
--- a/Engine/Source/Editor/IntroTutorials/Private/TutorialStateSettings.h
+++ b/Engine/Source/Editor/IntroTutorials/Private/TutorialStateSettings.h
@@ -41,6 +41,9 @@ class UTutorialStateSettings : public UObject
 	/** UObject interface */
 	virtual void PostInitProperties() override;
 
+	/** Reset the progress and completion sate of all tutorials */
+	void ClearProgress();
+
 	/** Get the recorded progress of the pass-in tutorial */
 	int32 GetProgress(UEditorTutorial* InTutorial, bool& bOutHaveSeenTutorial) const;
 
diff --git a/Engine/Source/Editor/IntroTutorials/Private/TutorialText.cpp b/Engine/Source/Editor/IntroTutorials/Private/TutorialText.cpp
--- a/Engine/Source/Editor/IntroTutorials/Private/TutorialText.cpp
+++ b/Engine/Source/Editor/IntroTutorials/Private/TutorialText.cpp
@@ -10,6 +10,8 @@
 #include "EditorTutorial.h"
 #include "SourceCodeNavigation.h"
 #include "TutorialImageDecorator.h"
+#include "EngineAnalytics.h"
+#include "Runtime/Analytics/Analytics/Public/Interfaces/IAnalyticsProvider.h"
 
 #define LOCTEXT_NAMESPACE "TutorialText"
 
@@ -61,6 +63,14 @@ static void OnBrowserLinkClicked(const FSlateHyperlinkRun::FMetadata& Metadata)
 	const FString* Url = Metadata.Find(TEXT("href"));
 	if(Url)
 	{
+		if( FEngineAnalytics::IsAvailable() )
+		{
+			TArray<FAnalyticsEventAttribute> EventAttributes;
+			EventAttributes.Add(FAnalyticsEventAttribute(TEXT("BrowserLink"), *Url));
+
+			FEngineAnalytics::GetProvider().RecordEvent( TEXT("Rocket.Tutorials.BrowserLinkClicked"), EventAttributes );
+		}
+
 		FPlatformProcess::LaunchURL(**Url, nullptr, nullptr);
 	}
 }
@@ -70,6 +80,14 @@ static void OnDocLinkClicked(const FSlateHyperlinkRun::FMetadata& Metadata)
 	const FString* Url = Metadata.Find(TEXT("href"));
 	if(Url)
 	{
+		if( FEngineAnalytics::IsAvailable() )
+		{
+			TArray<FAnalyticsEventAttribute> EventAttributes;
+			EventAttributes.Add(FAnalyticsEventAttribute(TEXT("DocLink"), *Url));
+
+			FEngineAnalytics::GetProvider().RecordEvent( TEXT("Rocket.Tutorials.DocLinkClicked"), EventAttributes );
+		}
+
 		IDocumentation::Get()->Open(*Url);
 	}
 }
@@ -83,6 +101,14 @@ static void ParseTutorialLink(const FString &InternalLink)
 		IMainFrameModule& MainFrameModule = FModuleManager::LoadModuleChecked<IMainFrameModule>(TEXT("MainFrame"));
 		const bool bRestart = true;
 		IntroTutorials.LaunchTutorial(Blueprint->GeneratedClass->GetDefaultObject<UEditorTutorial>(), bRestart, MainFrameModule.GetParentWindow());
+
+		if( FEngineAnalytics::IsAvailable() )
+		{
+			TArray<FAnalyticsEventAttribute> EventAttributes;
+			EventAttributes.Add(FAnalyticsEventAttribute(TEXT("TutorialLink"), InternalLink));
+
+			FEngineAnalytics::GetProvider().RecordEvent( TEXT("Rocket.Tutorials.TutorialLinkClicked"), EventAttributes );
+		}
 	}
 }
 
@@ -146,6 +172,14 @@ static void ParseCodeLink(const FString &InternalLink)
 	Path = FPaths::ConvertRelativePathToFull(Path);
 
 	SourceCodeAccessor.OpenFileAtLine(Path, Line, Col);
+
+	if( FEngineAnalytics::IsAvailable() )
+	{
+		TArray<FAnalyticsEventAttribute> EventAttributes;
+		EventAttributes.Add(FAnalyticsEventAttribute(TEXT("CodeLink"), InternalLink));
+
+		FEngineAnalytics::GetProvider().RecordEvent( TEXT("Rocket.Tutorials.CodeLinkClicked"), EventAttributes );
+	}
 }
 
 static void OnCodeLinkClicked(const FSlateHyperlinkRun::FMetadata& Metadata)
@@ -173,6 +207,14 @@ static void ParseAssetLink(const FString& InternalLink, const FString* Action)
 		{
 			FAssetEditorManager::Get().OpenEditorForAsset(RequiredObject);
 		}
+
+		if( FEngineAnalytics::IsAvailable() )
+		{
+			TArray<FAnalyticsEventAttribute> EventAttributes;
+			EventAttributes.Add(FAnalyticsEventAttribute(TEXT("AssetLink"), InternalLink));
+
+			FEngineAnalytics::GetProvider().RecordEvent( TEXT("Rocket.Tutorials.AssetLinkClicked"), EventAttributes );
+		}
 	}
 }
 
diff --git a/Engine/Source/Editor/IntroTutorials/Public/IIntroTutorials.h b/Engine/Source/Editor/IntroTutorials/Public/IIntroTutorials.h
--- a/Engine/Source/Editor/IntroTutorials/Public/IIntroTutorials.h
+++ b/Engine/Source/Editor/IntroTutorials/Public/IIntroTutorials.h
@@ -36,21 +36,6 @@ public:
 		return FModuleManager::Get().IsModuleLoaded( "IntroTutorials" );
 	}
 
-	/**
-	 * Adds a tutorial association for a type of asset (will be summoned the first time an asset editor is opened for the specified asset class)
-	 * @param AssetClass		Asset type to register a tutorial for
-	 * @param TutorialDocPath	The path to the tutorial doc (must be rooted in the Shared directory, e.g., "Shared/Tutorials/InPersonaAnimEditorTutorial")
-	 * @param TutorialHasBeenSeenSettingName	The setting name to store whether this editor has been seen or not yet (stored in GEditorGameAgnosticIni in the [IntroTutorials] section)
-	 * @param SurveyGUIDString	Text representation of a GUID, used to uniquely identify this tutorial if required for a survey.
-	 */
-	virtual void RegisterTutorialForAssetEditor(UClass* AssetClass, const FString& TutorialDocPath, const FString& TutorialHasBeenSeenSettingName, const FString& SurveyGUIDString)=0;
-
-	/**
-	 * Removes the tutorial association for the specified class
-	 * @param AssetClass		Asset type to unregister the tutorial from
-	 */
-	virtual void UnregisterTutorialForAssetEditor(UClass* AssetClass)=0;
-
 	/**
 	 * Launch a tutorial immediately, bypassing the tutorial browser.
 	 * 
diff --git a/Engine/Source/Editor/Kismet/Private/BlueprintEditor.cpp b/Engine/Source/Editor/Kismet/Private/BlueprintEditor.cpp
--- a/Engine/Source/Editor/Kismet/Private/BlueprintEditor.cpp
+++ b/Engine/Source/Editor/Kismet/Private/BlueprintEditor.cpp
@@ -5947,21 +5947,27 @@ FName FBlueprintEditor::GetToolkitFName() const
 	return FName("BlueprintEditor");
 }
 
+FName FBlueprintEditor::GetContextFromBlueprintType(EBlueprintType InType)
+{
+	switch (InType)
+	{
+	default:
+	case BPTYPE_Normal:
+		return FName("BlueprintEditor");
+	case BPTYPE_MacroLibrary:
+		return FName("BlueprintEditor.MacroLibrary");
+	case BPTYPE_Interface:
+		return FName("BlueprintEditor.Interface");
+	case BPTYPE_LevelScript:
+		return FName("BlueprintEditor.LevelScript");
+	}
+}
+
 FName FBlueprintEditor::GetToolkitContextFName() const
 {
 	if(GetBlueprintObj())
 	{
-		switch (GetBlueprintObj()->BlueprintType)
-		{
-		case BPTYPE_Normal:
-			return FName("BlueprintEditor");
-		case BPTYPE_MacroLibrary:
-			return FName("BlueprintEditor.MacroLibrary");
-		case BPTYPE_Interface:
-			return FName("BlueprintEditor.Interface");
-		case BPTYPE_LevelScript:
-			return FName("BlueprintEditor.LevelScript");
-		}
+		return GetContextFromBlueprintType(GetBlueprintObj()->BlueprintType);
 	}
 
 	return FName("BlueprintEditor");
diff --git a/Engine/Source/Editor/Kismet/Public/BlueprintEditor.h b/Engine/Source/Editor/Kismet/Public/BlueprintEditor.h
--- a/Engine/Source/Editor/Kismet/Public/BlueprintEditor.h
+++ b/Engine/Source/Editor/Kismet/Public/BlueprintEditor.h
@@ -489,6 +489,9 @@ public:
 	/** Checks to see if the provided graph is contained within the current blueprint */
 	bool IsGraphInCurrentBlueprint(UEdGraph* InGraph) const;
 
+	/** Get the context to use from the Blueprint type */
+	static FName GetContextFromBlueprintType(EBlueprintType InType);
+
 protected:
 	/** Called during initialization of the blueprint editor to register any application modes. */
 	virtual void RegisterApplicationModes(const TArray<UBlueprint*>& InBlueprints, bool bShouldOpenInDefaultsMode);
diff --git a/Engine/Source/Editor/MainFrame/Private/Frame/MainFrameActions.cpp b/Engine/Source/Editor/MainFrame/Private/Frame/MainFrameActions.cpp
--- a/Engine/Source/Editor/MainFrame/Private/Frame/MainFrameActions.cpp
+++ b/Engine/Source/Editor/MainFrame/Private/Frame/MainFrameActions.cpp
@@ -523,9 +523,9 @@ void FMainFrameActionCallbacks::PackageProject( const FName InPlatformInfoName )
 		const ITargetPlatform* const Platform = GetTargetPlatformManager()->FindTargetPlatform(PlatformInfo->TargetPlatformName.ToString());
 		if (Platform)
 		{
-			FString NotInstalledDocLink;
+			FString NotInstalledTutorialLink;
 			FString ProjectPath = FPaths::IsProjectFilePathSet() ? FPaths::ConvertRelativePathToFull(FPaths::GetProjectFilePath()) : FPaths::RootDir() / FApp::GetGameName() / FApp::GetGameName() + TEXT(".uproject");
-			int32 Result = Platform->DoesntHaveRequirements(ProjectPath, bProjectHasCode, NotInstalledDocLink);
+			int32 Result = Platform->DoesntHaveRequirements(ProjectPath, bProjectHasCode, NotInstalledTutorialLink);
 
 			// report to analytics
 			FEditorAnalytics::ReportBuildRequirementsFailure(TEXT("Editor.Package.Failed"), PlatformInfo->TargetPlatformName.ToString(), bProjectHasCode, Result);
@@ -577,7 +577,7 @@ void FMainFrameActionCallbacks::PackageProject( const FName InPlatformInfoName )
 				{
 					// broadcast this, and assume someone will pick it up
 					IMainFrameModule& MainFrameModule = FModuleManager::GetModuleChecked<IMainFrameModule>(TEXT("MainFrame"));
-					MainFrameModule.BroadcastMainFrameSDKNotInstalled(PlatformInfo->TargetPlatformName.ToString(), NotInstalledDocLink);
+					MainFrameModule.BroadcastMainFrameSDKNotInstalled(PlatformInfo->TargetPlatformName.ToString(), NotInstalledTutorialLink);
 				}
 				return;
 
