CIS fix

diff --git a/Engine/Source/Editor/IntroTutorials/Private/STutorialButton.cpp b/Engine/Source/Editor/IntroTutorials/Private/STutorialButton.cpp
--- a/Engine/Source/Editor/IntroTutorials/Private/STutorialButton.cpp
+++ b/Engine/Source/Editor/IntroTutorials/Private/STutorialButton.cpp
@@ -149,7 +149,7 @@ FReply STutorialButton::HandleButtonClicked()
 			TArray<FAnalyticsEventAttribute> EventAttributes;
 			EventAttributes.Add(FAnalyticsEventAttribute(TEXT("Context"), Context.ToString()));
 			EventAttributes.Add(FAnalyticsEventAttribute(TEXT("TimeSinceAlertStarted"), (AlertStartTime != 0.0f && ShouldShowAlert()) ? (FPlatformTime::Seconds() - AlertStartTime) : -1.0f));
-			EventAttributes.Add(FAnalyticsEventAttribute(TEXT("LaunchedBrowser"), ShouldLaunchBrowser());
+			EventAttributes.Add(FAnalyticsEventAttribute(TEXT("LaunchedBrowser"), ShouldLaunchBrowser()));
 
 			FEngineAnalytics::GetProvider().RecordEvent( TEXT("Rocket.Tutorials.ClickedContextButton"), EventAttributes );
 		}
