Widget-bound content that is no visible now falls back to rendering centred.

Previously it would just disappear, whihc was very confusing for users if they had not read the previous stage.

diff --git a/Engine/Source/Editor/IntroTutorials/Private/STutorialContent.cpp b/Engine/Source/Editor/IntroTutorials/Private/STutorialContent.cpp
--- a/Engine/Source/Editor/IntroTutorials/Private/STutorialContent.cpp
+++ b/Engine/Source/Editor/IntroTutorials/Private/STutorialContent.cpp
@@ -334,47 +334,57 @@ TSharedRef<SWidget> STutorialContent::GenerateContentWidget(const FTutorialConte
 
 FVector2D STutorialContent::GetPosition() const
 {
-	float XOffset = 0.0f;
-	switch(HorizontalAlignment.Get())
+	const bool bNonVisibleWidgetBound = !bIsVisible && Anchor.Type == ETutorialAnchorIdentifier::NamedWidget;
+	if(bNonVisibleWidgetBound)
 	{
-	case HAlign_Left:
-		XOffset = -(ContentWidget->GetDesiredSize().X - ContentOffset);
-		break;
-	default:
-	case HAlign_Fill:
-	case HAlign_Center:
-		XOffset = (CachedGeometry.Size.X * 0.5f) - (ContentWidget->GetDesiredSize().X * 0.5f);
-		break;
-	case HAlign_Right:
-		XOffset = CachedGeometry.Size.X - ContentOffset;
-		break;
+		// fallback: center on cached window
+		return FVector2D((CachedWindowSize.X * 0.5f) - (ContentWidget->GetDesiredSize().X * 0.5f),
+						 (CachedWindowSize.Y * 0.5f) - (ContentWidget->GetDesiredSize().Y * 0.5f));
 	}
+	else
+	{
+		float XOffset = 0.0f;
+		switch(HorizontalAlignment.Get())
+		{
+		case HAlign_Left:
+			XOffset = -(ContentWidget->GetDesiredSize().X - ContentOffset);
+			break;
+		default:
+		case HAlign_Fill:
+		case HAlign_Center:
+			XOffset = (CachedGeometry.Size.X * 0.5f) - (ContentWidget->GetDesiredSize().X * 0.5f);
+			break;
+		case HAlign_Right:
+			XOffset = CachedGeometry.Size.X - ContentOffset;
+			break;
+		}
 
-	XOffset += WidgetOffset.Get().X;
+		XOffset += WidgetOffset.Get().X;
 
-	float YOffset = 0.0f;
-	switch(VerticalAlignment.Get())
-	{
-	case VAlign_Top:
-		YOffset = -(ContentWidget->GetDesiredSize().Y - ContentOffset);
-		break;
-	default:
-	case VAlign_Fill:
-	case VAlign_Center:
-		YOffset = (CachedGeometry.Size.Y * 0.5f) - (ContentWidget->GetDesiredSize().Y * 0.5f);
-		break;
-	case VAlign_Bottom:
-		YOffset = (CachedGeometry.Size.Y - ContentOffset);
-		break;
-	}
+		float YOffset = 0.0f;
+		switch(VerticalAlignment.Get())
+		{
+		case VAlign_Top:
+			YOffset = -(ContentWidget->GetDesiredSize().Y - ContentOffset);
+			break;
+		default:
+		case VAlign_Fill:
+		case VAlign_Center:
+			YOffset = (CachedGeometry.Size.Y * 0.5f) - (ContentWidget->GetDesiredSize().Y * 0.5f);
+			break;
+		case VAlign_Bottom:
+			YOffset = (CachedGeometry.Size.Y - ContentOffset);
+			break;
+		}
 
-	YOffset += WidgetOffset.Get().Y;
+		YOffset += WidgetOffset.Get().Y;
 
-	// now build & clamp to area
-	FVector2D BaseOffset = FVector2D(CachedGeometry.AbsolutePosition.X + XOffset, CachedGeometry.AbsolutePosition.Y + YOffset);
-	BaseOffset.X = FMath::Clamp(BaseOffset.X, 0.0f, CachedWindowSize.X - ContentWidget->GetDesiredSize().X);
-	BaseOffset.Y = FMath::Clamp(BaseOffset.Y, 0.0f, CachedWindowSize.Y - ContentWidget->GetDesiredSize().Y);
-	return BaseOffset;
+		// now build & clamp to area
+		FVector2D BaseOffset = FVector2D(CachedGeometry.AbsolutePosition.X + XOffset, CachedGeometry.AbsolutePosition.Y + YOffset);
+		BaseOffset.X = FMath::Clamp(BaseOffset.X, 0.0f, CachedWindowSize.X - ContentWidget->GetDesiredSize().X);
+		BaseOffset.Y = FMath::Clamp(BaseOffset.Y, 0.0f, CachedWindowSize.Y - ContentWidget->GetDesiredSize().Y);
+		return BaseOffset;
+	}
 }
 
 FVector2D STutorialContent::GetSize() const
@@ -429,7 +439,13 @@ void STutorialContent::HandleCacheWindowSize(const FVector2D& InWindowSize)
 
 EVisibility STutorialContent::GetVisibility() const
 {
-	return bIsVisible || Anchor.Type == ETutorialAnchorIdentifier::None ? EVisibility::SelfHitTestInvisible : EVisibility::Collapsed;
+	const bool bVisibleWidgetBound = bIsVisible && Anchor.Type == ETutorialAnchorIdentifier::NamedWidget;
+	const bool bNonWidgetBound =  Anchor.Type == ETutorialAnchorIdentifier::None;
+
+	// fallback if widget is not drawn - we should display this content anyway
+	const bool bNonVisibleWidgetBound = !bIsVisible && Anchor.Type == ETutorialAnchorIdentifier::NamedWidget;
+
+	return (bVisibleWidgetBound || bNonWidgetBound || bNonVisibleWidgetBound) ? EVisibility::SelfHitTestInvisible : EVisibility::Collapsed;
 }
 
 FSlateColor STutorialContent::GetBackgroundColor() const
