Tutorial button now provides more context menu options

You can access all context-sensitive options via the right-click menu.

Also exposed the ability to refresh available tutorials (as they may have bene reloaded since).

diff --git a/Engine/Source/Editor/IntroTutorials/Private/SEditorTutorials.cpp b/Engine/Source/Editor/IntroTutorials/Private/SEditorTutorials.cpp
--- a/Engine/Source/Editor/IntroTutorials/Private/SEditorTutorials.cpp
+++ b/Engine/Source/Editor/IntroTutorials/Private/SEditorTutorials.cpp
@@ -139,6 +139,7 @@ void SEditorTutorials::HandleHomeClicked()
 
 	ContentBox->ClearChildren();
 	bBrowserVisible = true;
+	TutorialHome->ReloadTutorials();
 	OnTutorialExited.ExecuteIfBound();
 }
 
diff --git a/Engine/Source/Editor/IntroTutorials/Private/STutorialButton.cpp b/Engine/Source/Editor/IntroTutorials/Private/STutorialButton.cpp
--- a/Engine/Source/Editor/IntroTutorials/Private/STutorialButton.cpp
+++ b/Engine/Source/Editor/IntroTutorials/Private/STutorialButton.cpp
@@ -68,6 +68,11 @@ void STutorialButton::Tick(const FGeometry& AllottedGeometry, const double InCur
 		{
 			AlertStartTime = FPlatformTime::Seconds();
 		}
+
+		if(LaunchTutorial != nullptr)
+		{
+			TutorialTitle = LaunchTutorial->Title;
+		}
 	}
 	bDeferTutorialOpen = false;
 }
@@ -171,33 +176,40 @@ FReply STutorialButton::HandleButtonClicked()
 
 FReply STutorialButton::OnMouseButtonDown(const FGeometry& MyGeometry, const FPointerEvent& MouseEvent)
 {
-	if(!ShouldLaunchBrowser())
+	if (MouseEvent.GetEffectingButton() == EKeys::RightMouseButton)
 	{
-		if (MouseEvent.GetEffectingButton() == EKeys::RightMouseButton)
+		const bool bInShouldCloseWindowAfterMenuSelection = true;
+		FMenuBuilder MenuBuilder(bInShouldCloseWindowAfterMenuSelection, nullptr);
+
+		if(ShouldShowAlert())
 		{
-			const bool bInShouldCloseWindowAfterMenuSelection = true;
-			FMenuBuilder MenuBuilder(bInShouldCloseWindowAfterMenuSelection, nullptr);
-
-			if(ShouldShowAlert())
-			{
-				MenuBuilder.AddMenuEntry(
-					LOCTEXT("DismissReminder", "Dismiss Alert"),
-					LOCTEXT("DismissReminderTooltip", "Don't show me this alert again"),
-					FSlateIcon(),
-					FUIAction(FExecuteAction::CreateSP(this, &STutorialButton::DismissAlert))
-					);
-			}
+			MenuBuilder.AddMenuEntry(
+				LOCTEXT("DismissReminder", "Dismiss Alert"),
+				LOCTEXT("DismissReminderTooltip", "Don't show me this alert again"),
+				FSlateIcon(),
+				FUIAction(FExecuteAction::CreateSP(this, &STutorialButton::DismissAlert))
+				);
+		}
 
+		if(bTutorialAvailable)
+		{
 			MenuBuilder.AddMenuEntry(
-				LOCTEXT("LaunchBrowser", "Show Available Tutorials"),
-				LOCTEXT("LaunchBrowserTooltip", "Display the tutorials browser"),
+				FText::Format(LOCTEXT("LaunchTutorialPattern", "Open Tutorial: {0}"), TutorialTitle),
+				LOCTEXT("LaunchTutorialTooltip", "Launch this tutorial"),
 				FSlateIcon(),
-				FUIAction(FExecuteAction::CreateSP(this, &STutorialButton::LaunchBrowser))
+				FUIAction(FExecuteAction::CreateSP(this, &STutorialButton::LaunchTutorial))
 				);
+		}
 
+		MenuBuilder.AddMenuEntry(
+			LOCTEXT("LaunchBrowser", "Show Available Tutorials"),
+			LOCTEXT("LaunchBrowserTooltip", "Display the tutorials browser"),
+			FSlateIcon(),
+			FUIAction(FExecuteAction::CreateSP(this, &STutorialButton::LaunchBrowser))
+			);
 
-			FSlateApplication::Get().PushMenu(SharedThis(this), MenuBuilder.MakeWidget(), FSlateApplication::Get().GetCursorPos(), FPopupTransitionEffect::ContextMenu);
-		}
+
+		FSlateApplication::Get().PushMenu(SharedThis(this), MenuBuilder.MakeWidget(), FSlateApplication::Get().GetCursorPos(), FPopupTransitionEffect::ContextMenu);
 	}
 	return FReply::Handled();
 }
@@ -229,6 +241,11 @@ void STutorialButton::DismissAlert()
 	}
 }
 
+void STutorialButton::LaunchTutorial()
+{
+	HandleButtonClicked();
+}
+
 void STutorialButton::LaunchBrowser()
 {
 	if(ContextWindow.IsValid())
@@ -259,8 +276,12 @@ FText STutorialButton::GetButtonToolTip() const
 	{
 		return LOCTEXT("TutorialLaunchBrowserToolTip", "Show Available Tutorials");
 	}
-
-	return LOCTEXT("TutorialLaunchToolTip", "Launch Tutorial (right-click for more options)");
+	else if(bTutorialAvailable)
+	{
+		return FText::Format(LOCTEXT("TutorialLaunchToolTipPattern", "Open: {0}\nRight-Click for More Options"), TutorialTitle);
+	}
+	
+	return LOCTEXT("TutorialToolTip", "Take Tutorial");
 }
 
 #undef LOCTEXT_NAMESPACE
\ No newline at end of file
diff --git a/Engine/Source/Editor/IntroTutorials/Private/STutorialButton.h b/Engine/Source/Editor/IntroTutorials/Private/STutorialButton.h
--- a/Engine/Source/Editor/IntroTutorials/Private/STutorialButton.h
+++ b/Engine/Source/Editor/IntroTutorials/Private/STutorialButton.h
@@ -38,6 +38,9 @@ private:
 	/** Get the tooltip for the tutorials button */
 	FText GetButtonToolTip() const;
 
+	/** Launch tutorial from the context menu */
+	void LaunchTutorial();
+
 private:
 	/** Flag to defer tutorial open until the first Tick() */
 	bool bDeferTutorialOpen;
@@ -62,4 +65,7 @@ private:
 
 	/** Start time we began playing the alert animation */
 	float AlertStartTime;
+
+	/**	The name of the tutorial we will launch */
+	FText TutorialTitle;
 };
\ No newline at end of file
