Added completion 'Tick' displayed in the browser

Shows when a tutorial has been completed.
Also added tooltips explaining what the icon means & a tooltip for the restart button.

diff --git a/Engine/Source/Editor/EditorStyle/Private/SlateEditorStyle.cpp b/Engine/Source/Editor/EditorStyle/Private/SlateEditorStyle.cpp
--- a/Engine/Source/Editor/EditorStyle/Private/SlateEditorStyle.cpp
+++ b/Engine/Source/Editor/EditorStyle/Private/SlateEditorStyle.cpp
@@ -2500,6 +2500,8 @@ void FSlateEditorStyle::FStyle::SetupTutorialStyles()
 		Set( "Tutorials.Browser.BackButton", new IMAGE_BRUSH( "Tutorials/BrowserBack", FVector2D(32.0f, 32.0f), FLinearColor(0.0f, 0.0f, 0.0f, 1.0f) ) );
 		Set( "Tutorials.Browser.RestartButton", new IMAGE_BRUSH( "Tutorials/BrowserRestart", FVector2D(16.0f, 16.0f), FLinearColor(0.0f, 0.0f, 0.0f, 1.0f) ) );
 
+		Set( "Tutorials.Browser.Completed", new IMAGE_BRUSH( "Tutorials/TutorialCompleted", Icon32x32 ) );
+
 		Set( "Tutorials.Navigation.Button", FButtonStyle()
 			.SetNormal( BOX_BRUSH( "Common/ButtonHoverHint", FMargin(4/16.0f), FLinearColor(0,0,0,0) ) )
 			.SetHovered( BOX_BRUSH( "Common/ButtonHoverHint", FMargin(4/16.0f), FLinearColor(0,0,0,0) ) )
diff --git a/Engine/Source/Editor/IntroTutorials/Private/STutorialsBrowser.cpp b/Engine/Source/Editor/IntroTutorials/Private/STutorialsBrowser.cpp
--- a/Engine/Source/Editor/IntroTutorials/Private/STutorialsBrowser.cpp
+++ b/Engine/Source/Editor/IntroTutorials/Private/STutorialsBrowser.cpp
@@ -236,6 +236,11 @@ public:
 
 	virtual TSharedRef<ITableRow> OnGenerateTutorialRow(const TSharedRef<STableViewBase>& OwnerTable) const override
 	{
+		bHaveCompletedTutorial = GetDefault<UTutorialStateSettings>()->HaveCompletedTutorial(Tutorial);
+		bHaveSeenTutorial = false;
+		const int32 CurrentStage = GetDefault<UTutorialStateSettings>()->GetProgress(Tutorial, bHaveSeenTutorial);
+		Progress = (Tutorial->Stages.Num() > 0) ? (float)(CurrentStage + 1) / (float)Tutorial->Stages.Num() : 0.0f;
+
 		return SNew(STableRow<TSharedPtr<ITutorialListEntry>>, OwnerTable)
 		[
 			SAssignNew(LaunchButton, SButton)
@@ -254,14 +259,27 @@ public:
 					.HAlign(HAlign_Center)
 					.Padding(8.0f)
 					[
-						SNew(SBox)
-						.WidthOverride(64.0f)
-						.HeightOverride(64.0f)
-						.VAlign(VAlign_Center)
-						.HAlign(HAlign_Center)
+						SNew(SOverlay)
+						+SOverlay::Slot()
+						[
+							SNew(SBox)
+							.WidthOverride(64.0f)
+							.HeightOverride(64.0f)
+							.VAlign(VAlign_Center)
+							.HAlign(HAlign_Center)
+							[
+								SNew(SImage)
+								.Image(SlateBrush)
+							]
+						]
+						+SOverlay::Slot()
+						.VAlign(VAlign_Bottom)
+						.HAlign(HAlign_Right)
 						[
 							SNew(SImage)
-							.Image(SlateBrush)
+							.ToolTipText(LOCTEXT("CompletedCheckToolTip", "This tutorial has been completed"))
+							.Visibility(this, &FTutorialListEntry_Tutorial::GetCompletedVisibility)
+							.Image(FEditorStyle::Get().GetBrush("Tutorials.Browser.Completed"))
 						]
 					]
 					+SHorizontalBox::Slot()
@@ -289,6 +307,7 @@ public:
 							.VAlign(VAlign_Center)
 							[
 								SNew(SButton)
+								.ToolTipText(LOCTEXT("RestartButtonToolTip", "Start this tutorial from the beginning"))
 								.Visibility(this, &FTutorialListEntry_Tutorial::GetProgressVisibility)
 								.OnClicked(this, &FTutorialListEntry_Tutorial::OnClicked, true)
 								.ButtonStyle(&FEditorStyle::Get().GetWidgetStyle<FButtonStyle>("Tutorials.Browser.Button"))
@@ -338,23 +357,24 @@ public:
 
 	TOptional<float> GetProgress() const
 	{
-		bool bHaveSeenTutorial = false;
-		const int32 CurrentStage = GetDefault<UTutorialStateSettings>()->GetProgress(Tutorial, bHaveSeenTutorial);
-		return (Tutorial->Stages.Num() > 0) ? (float)(CurrentStage + 1) / (float)Tutorial->Stages.Num() : 0.0f;
+		return Progress;
 	}
 
 	EVisibility GetProgressVisibility() const
 	{
 		if(LaunchButton->IsHovered())
 		{
-			bool bHaveSeenTutorial = false;
-			const int32 CurrentStage = GetDefault<UTutorialStateSettings>()->GetProgress(Tutorial, bHaveSeenTutorial);
 			return LaunchButton->IsHovered() && bHaveSeenTutorial ? EVisibility::Visible : EVisibility::Hidden;
 		}
 
 		return EVisibility::Hidden;
 	}
 
+	EVisibility GetCompletedVisibility() const
+	{
+		return bHaveCompletedTutorial ? EVisibility::Visible : EVisibility::Hidden;
+	}
+
 public:
 	/** Parent category */
 	TWeakPtr<ITutorialListEntry> ParentCategory;
@@ -379,6 +399,15 @@ public:
 
 	/** Dynamic brush from the texture specified by the user */
 	TSharedPtr<FSlateDynamicImageBrush> DynamicBrush;
+
+	/** Cached tutorial completion state */
+	mutable bool bHaveCompletedTutorial;
+
+	/** Cached tutorial seen state */
+	mutable bool bHaveSeenTutorial;
+
+	/** Cached tutorial progress */
+	mutable float Progress;
 };
 
 void STutorialsBrowser::Construct(const FArguments& InArgs)
