Fix perf regression introduced by tutorials

The tutorial overlay was traversing the widget hierarchy even if it had nothing to do.

diff --git a/Engine/Source/Editor/IntroTutorials/Private/STutorialOverlay.cpp b/Engine/Source/Editor/IntroTutorials/Private/STutorialOverlay.cpp
--- a/Engine/Source/Editor/IntroTutorials/Private/STutorialOverlay.cpp
+++ b/Engine/Source/Editor/IntroTutorials/Private/STutorialOverlay.cpp
@@ -10,12 +10,15 @@
 #include "BlueprintEditorUtils.h"
 #include "Guid.h"
 
+static FName IntroTutorialsModuleName("IntroTutorials");
+
 void STutorialOverlay::Construct(const FArguments& InArgs, UEditorTutorial* InTutorial, FTutorialStage* const InStage)
 {
 	ParentWindow = InArgs._ParentWindow;
 	bIsStandalone = InArgs._IsStandalone;
 	OnClosed = InArgs._OnClosed;
-	
+	bHasValidContent = InStage != nullptr;
+
 	TSharedPtr<SOverlay> Overlay;
 
 	ChildSlot
@@ -100,10 +103,21 @@ int32 STutorialOverlay::OnPaint( const FPaintArgs& Args, const FGeometry& Allott
 {
 	if(ParentWindow.IsValid())
 	{
-		TSharedPtr<SWindow> PinnedWindow = ParentWindow.Pin();
-		OnResetNamedWidget.Broadcast();
-		OnCacheWindowSize.Broadcast(PinnedWindow->GetWindowGeometryInWindow().Size);
-		LayerId = TraverseWidgets(PinnedWindow.ToSharedRef(), PinnedWindow->GetWindowGeometryInWindow(), MyClippingRect, OutDrawElements, LayerId);
+		bool bIsPicking = false;
+		FName WidgetNameToHighlight = NAME_None;
+		FIntroTutorials& IntroTutorials = FModuleManager::Get().GetModuleChecked<FIntroTutorials>(IntroTutorialsModuleName);
+		if(IntroTutorials.OnIsPicking().IsBound())
+		{
+			bIsPicking = IntroTutorials.OnIsPicking().Execute(WidgetNameToHighlight);
+		}
+
+		if(bIsPicking || bHasValidContent)
+		{
+			TSharedPtr<SWindow> PinnedWindow = ParentWindow.Pin();
+			OnResetNamedWidget.Broadcast();
+			OnCacheWindowSize.Broadcast(PinnedWindow->GetWindowGeometryInWindow().Size);
+			LayerId = TraverseWidgets(PinnedWindow.ToSharedRef(), PinnedWindow->GetWindowGeometryInWindow(), MyClippingRect, OutDrawElements, LayerId);
+		}
 	}
 	
 	return SCompoundWidget::OnPaint(Args, AllottedGeometry, MyClippingRect, OutDrawElements, LayerId, InWidgetStyle, bParentEnabled);
@@ -121,7 +135,6 @@ int32 STutorialOverlay::TraverseWidgets(TSharedRef<SWidget> InWidget, const FGeo
 		// if we are picking, we need to draw an outline here
 		FName WidgetNameToHighlight = NAME_None;
 		bool bIsPicking = false;
-		static FName IntroTutorialsModuleName("IntroTutorials");
 		FIntroTutorials& IntroTutorials = FModuleManager::Get().GetModuleChecked<FIntroTutorials>(IntroTutorialsModuleName);
 		if(IntroTutorials.OnIsPicking().IsBound())
 		{
diff --git a/Engine/Source/Editor/IntroTutorials/Private/STutorialOverlay.h b/Engine/Source/Editor/IntroTutorials/Private/STutorialOverlay.h
--- a/Engine/Source/Editor/IntroTutorials/Private/STutorialOverlay.h
+++ b/Engine/Source/Editor/IntroTutorials/Private/STutorialOverlay.h
@@ -84,4 +84,6 @@ private:
 	/** Delegate used to inform widgets of the current window size, so they can auto-adjust layout */
 	FOnCacheWindowSize OnCacheWindowSize;
 
+	/** Flag to see if we have valid content (this widget is created to also supply picker overlays) */
+	bool bHasValidContent;
 };
\ No newline at end of file
