Fixed tutorial alert dismissal not always working

Tutorial alert also hides when the tutorial is taken.

diff --git a/Engine/Source/Editor/IntroTutorials/Private/STutorialButton.cpp b/Engine/Source/Editor/IntroTutorials/Private/STutorialButton.cpp
--- a/Engine/Source/Editor/IntroTutorials/Private/STutorialButton.cpp
+++ b/Engine/Source/Editor/IntroTutorials/Private/STutorialButton.cpp
@@ -48,20 +48,14 @@ void STutorialButton::Tick(const FGeometry& AllottedGeometry, const double InCur
 {
 	if (bDeferTutorialOpen)
 	{
-		UEditorTutorial* AttractTutorial = nullptr;
-		UEditorTutorial* LaunchTutorial = nullptr;
-		FString BrowserFilter;
-		GetDefault<UEditorTutorialSettings>()->FindTutorialInfoForContext(Context, AttractTutorial, LaunchTutorial, BrowserFilter);
-
-		bTutorialAvailable = (LaunchTutorial != nullptr);
-		bTutorialCompleted = (LaunchTutorial != nullptr) && GetDefault<UTutorialStateSettings>()->HaveCompletedTutorial(LaunchTutorial);
-		bTutorialDismissed = (AttractTutorial != nullptr) && GetDefault<UTutorialStateSettings>()->IsTutorialDismissed(AttractTutorial);
-		if (bTutorialAvailable && AttractTutorial != nullptr && !bTutorialDismissed && !bTutorialCompleted)
+		RefreshStatus();
+
+		if (bTutorialAvailable && CachedAttractTutorial != nullptr && !bTutorialDismissed && !bTutorialCompleted)
 		{
 			// kick off the attract tutorial if the user hasn't dismissed it and hasn't completed it
 			FIntroTutorials& IntroTutorials = FModuleManager::GetModuleChecked<FIntroTutorials>(TEXT("IntroTutorials"));
 			const bool bRestart = true;
-			IntroTutorials.LaunchTutorial(AttractTutorial, bRestart, ContextWindow);
+			IntroTutorials.LaunchTutorial(CachedAttractTutorial, bRestart, ContextWindow);
 		}
 
 		if(ShouldShowAlert())
@@ -69,9 +63,9 @@ void STutorialButton::Tick(const FGeometry& AllottedGeometry, const double InCur
 			AlertStartTime = FPlatformTime::Seconds();
 		}
 
-		if(LaunchTutorial != nullptr)
+		if(CachedLaunchTutorial != nullptr)
 		{
-			TutorialTitle = LaunchTutorial->Title;
+			TutorialTitle = CachedLaunchTutorial->Title;
 		}
 	}
 	bDeferTutorialOpen = false;
@@ -138,14 +132,7 @@ int32 STutorialButton::OnPaint(const FPaintArgs& Args, const FGeometry& Allotted
 
 FReply STutorialButton::HandleButtonClicked()
 {
-	UEditorTutorial* AttractTutorial = nullptr;
-	UEditorTutorial* LaunchTutorial = nullptr;
-	FString BrowserFilter;
-	GetDefault<UEditorTutorialSettings>()->FindTutorialInfoForContext(Context, AttractTutorial, LaunchTutorial, BrowserFilter);
-
-	bTutorialAvailable = (LaunchTutorial != nullptr);
-	bTutorialCompleted = (LaunchTutorial != nullptr) && GetDefault<UTutorialStateSettings>()->HaveCompletedTutorial(LaunchTutorial);
-	bTutorialDismissed = (AttractTutorial != nullptr) && GetDefault<UTutorialStateSettings>()->IsTutorialDismissed(AttractTutorial);
+	RefreshStatus();
 
 	if(ContextWindow.IsValid())
 	{
@@ -162,12 +149,19 @@ FReply STutorialButton::HandleButtonClicked()
 		FIntroTutorials& IntroTutorials = FModuleManager::GetModuleChecked<FIntroTutorials>(TEXT("IntroTutorials"));
 		if(ShouldLaunchBrowser())
 		{
-			IntroTutorials.SummonTutorialBrowser(ContextWindow.Pin().ToSharedRef(), BrowserFilter);
+			IntroTutorials.SummonTutorialBrowser(ContextWindow.Pin().ToSharedRef(), CachedBrowserFilter);
 		}
-		else if (LaunchTutorial != nullptr)
+		else if (CachedLaunchTutorial != nullptr)
 		{
+			auto Delegate = FSimpleDelegate::CreateSP(this, &STutorialButton::HandleTutorialExited);
+
 			const bool bRestart = true;
-			IntroTutorials.LaunchTutorial(LaunchTutorial, bRestart, ContextWindow);
+			IntroTutorials.LaunchTutorial(CachedLaunchTutorial, bRestart, ContextWindow, Delegate, Delegate);
+
+			const bool bDismissAcrossSessions = true;
+			GetMutableDefault<UTutorialStateSettings>()->DismissTutorial(CachedLaunchTutorial, bDismissAcrossSessions);
+			GetMutableDefault<UTutorialStateSettings>()->SaveProgress();
+			bTutorialDismissed = true;
 		}
 	}
 
@@ -216,29 +210,31 @@ FReply STutorialButton::OnMouseButtonDown(const FGeometry& MyGeometry, const FPo
 
 void STutorialButton::DismissAlert()
 {
-	UEditorTutorial* AttractTutorial = nullptr;
-	UEditorTutorial* LaunchTutorial = nullptr;
-	FString BrowserFilter;
-	GetDefault<UEditorTutorialSettings>()->FindTutorialInfoForContext(Context, AttractTutorial, LaunchTutorial, BrowserFilter);
-	if (AttractTutorial != nullptr)
-	{
-		if( FEngineAnalytics::IsAvailable() )
-		{
-			TArray<FAnalyticsEventAttribute> EventAttributes;
-			EventAttributes.Add(FAnalyticsEventAttribute(TEXT("Context"), Context.ToString()));
-			EventAttributes.Add(FAnalyticsEventAttribute(TEXT("TimeSinceAlertStarted"), (AlertStartTime != 0.0f && ShouldShowAlert()) ? (FPlatformTime::Seconds() - AlertStartTime) : -1.0f));
+	RefreshStatus();
 
-			FEngineAnalytics::GetProvider().RecordEvent( TEXT("Rocket.Tutorials.DismissedTutorialAlert"), EventAttributes );
-		}
+	if( FEngineAnalytics::IsAvailable() )
+	{
+		TArray<FAnalyticsEventAttribute> EventAttributes;
+		EventAttributes.Add(FAnalyticsEventAttribute(TEXT("Context"), Context.ToString()));
+		EventAttributes.Add(FAnalyticsEventAttribute(TEXT("TimeSinceAlertStarted"), (AlertStartTime != 0.0f && ShouldShowAlert()) ? (FPlatformTime::Seconds() - AlertStartTime) : -1.0f));
 
-		const bool bDismissAcrossSessions = true;
-		GetMutableDefault<UTutorialStateSettings>()->DismissTutorial(AttractTutorial, bDismissAcrossSessions);
-		GetMutableDefault<UTutorialStateSettings>()->SaveProgress();
-		bTutorialDismissed = true;
+		FEngineAnalytics::GetProvider().RecordEvent( TEXT("Rocket.Tutorials.DismissedTutorialAlert"), EventAttributes );
+	}
 
-		FIntroTutorials& IntroTutorials = FModuleManager::GetModuleChecked<FIntroTutorials>(TEXT("IntroTutorials"));
-		IntroTutorials.CloseAllTutorialContent();
+	const bool bDismissAcrossSessions = true;
+	if (CachedAttractTutorial != nullptr)
+	{
+		GetMutableDefault<UTutorialStateSettings>()->DismissTutorial(CachedAttractTutorial, bDismissAcrossSessions);
 	}
+	if( CachedLaunchTutorial != nullptr)
+	{
+		GetMutableDefault<UTutorialStateSettings>()->DismissTutorial(CachedLaunchTutorial, bDismissAcrossSessions);
+	}
+	GetMutableDefault<UTutorialStateSettings>()->SaveProgress();
+	bTutorialDismissed = true;
+
+	FIntroTutorials& IntroTutorials = FModuleManager::GetModuleChecked<FIntroTutorials>(TEXT("IntroTutorials"));
+	IntroTutorials.CloseAllTutorialContent();
 }
 
 void STutorialButton::LaunchTutorial()
@@ -250,13 +246,10 @@ void STutorialButton::LaunchBrowser()
 {
 	if(ContextWindow.IsValid())
 	{
-		UEditorTutorial* AttractTutorial = nullptr;
-		UEditorTutorial* LaunchTutorial = nullptr;
-		FString BrowserFilter;
-		GetDefault<UEditorTutorialSettings>()->FindTutorialInfoForContext(Context, AttractTutorial, LaunchTutorial, BrowserFilter);
+		RefreshStatus();
 
 		FIntroTutorials& IntroTutorials = FModuleManager::GetModuleChecked<FIntroTutorials>(TEXT("IntroTutorials"));
-		IntroTutorials.SummonTutorialBrowser(ContextWindow.Pin().ToSharedRef(), BrowserFilter);
+		IntroTutorials.SummonTutorialBrowser(ContextWindow.Pin().ToSharedRef(), CachedBrowserFilter);
 	}
 }
 
@@ -284,4 +277,22 @@ FText STutorialButton::GetButtonToolTip() const
 	return LOCTEXT("TutorialToolTip", "Take Tutorial");
 }
 
+void STutorialButton::RefreshStatus()
+{
+	CachedAttractTutorial = nullptr;
+	CachedLaunchTutorial = nullptr;
+	CachedBrowserFilter = TEXT("");
+	GetDefault<UEditorTutorialSettings>()->FindTutorialInfoForContext(Context, CachedAttractTutorial, CachedLaunchTutorial, CachedBrowserFilter);
+
+	bTutorialAvailable = (CachedLaunchTutorial != nullptr);
+	bTutorialCompleted = (CachedLaunchTutorial != nullptr) && GetDefault<UTutorialStateSettings>()->HaveCompletedTutorial(CachedLaunchTutorial);
+	bTutorialDismissed = ((CachedAttractTutorial != nullptr) && GetDefault<UTutorialStateSettings>()->IsTutorialDismissed(CachedAttractTutorial)) ||
+						 ((CachedLaunchTutorial != nullptr) && GetDefault<UTutorialStateSettings>()->IsTutorialDismissed(CachedLaunchTutorial));
+}
+
+void STutorialButton::HandleTutorialExited()
+{
+	RefreshStatus();
+}
+
 #undef LOCTEXT_NAMESPACE
\ No newline at end of file
diff --git a/Engine/Source/Editor/IntroTutorials/Private/STutorialButton.h b/Engine/Source/Editor/IntroTutorials/Private/STutorialButton.h
--- a/Engine/Source/Editor/IntroTutorials/Private/STutorialButton.h
+++ b/Engine/Source/Editor/IntroTutorials/Private/STutorialButton.h
@@ -41,6 +41,12 @@ private:
 	/** Launch tutorial from the context menu */
 	void LaunchTutorial();
 
+	/** Refresh internal status of flags, tutorials, filters etc.*/
+	void RefreshStatus();
+
+	/** Handle tutorial exiting/finishing */
+	void HandleTutorialExited();
+
 private:
 	/** Flag to defer tutorial open until the first Tick() */
 	bool bDeferTutorialOpen;
@@ -68,4 +74,13 @@ private:
 
 	/**	The name of the tutorial we will launch */
 	FText TutorialTitle;
+
+	/** Cached attract tutorial */
+	UEditorTutorial* CachedAttractTutorial;
+
+	/** Cached launch tutorial */
+	UEditorTutorial* CachedLaunchTutorial;
+
+	/** Cached browser filter */
+	FString CachedBrowserFilter;
 };
\ No newline at end of file
