Fixed tutorials border blocking mouse input to widgets below

diff --git a/Engine/Source/Editor/IntroTutorials/Private/STutorialContent.cpp b/Engine/Source/Editor/IntroTutorials/Private/STutorialContent.cpp
--- a/Engine/Source/Editor/IntroTutorials/Private/STutorialContent.cpp
+++ b/Engine/Source/Editor/IntroTutorials/Private/STutorialContent.cpp
@@ -76,6 +76,7 @@ void STutorialContent::Construct(const FArguments& InArgs, UEditorTutorial* InTu
 
 				// Add more padding if the content is to be displayed centrally (i.e. not on a widget)
 				.Padding(bIsStandalone ? TutorialConstants::BorderSizeStandalone : TutorialConstants::BorderSize)
+				.Visibility(EVisibility::SelfHitTestInvisible)
 				.BorderImage(FEditorStyle::GetBrush("Tutorials.Border"))
 				.BorderBackgroundColor(this, &STutorialContent::GetBackgroundColor)
 				.ForegroundColor(FCoreStyle::Get().GetSlateColor("InvertedForeground"))
