Fixed 'fallback' overlays appearing in multiple windows

Dont present the 'fallback' overlays unless in the primary ('navigation') window.

diff --git a/Engine/Source/Editor/IntroTutorials/Private/SEditorTutorials.cpp b/Engine/Source/Editor/IntroTutorials/Private/SEditorTutorials.cpp
--- a/Engine/Source/Editor/IntroTutorials/Private/SEditorTutorials.cpp
+++ b/Engine/Source/Editor/IntroTutorials/Private/SEditorTutorials.cpp
@@ -12,7 +12,7 @@
 void SEditorTutorials::Construct(const FArguments& InArgs)
 {
 	bBrowserVisible = false;
-	bShowNavigation = false;
+	bIsNavigationWindow = false;
 	ParentWindow = InArgs._ParentWindow;
 	OnNextClicked = InArgs._OnNextClicked;
 	OnBackClicked = InArgs._OnBackClicked;
@@ -82,7 +82,7 @@ void SEditorTutorials::Construct(const FArguments& InArgs)
 
 void SEditorTutorials::LaunchTutorial(bool bInIsNavigationWindow, FSimpleDelegate InOnTutorialClosed, FSimpleDelegate InOnTutorialExited)
 {
-	bShowNavigation = bInIsNavigationWindow;
+	bIsNavigationWindow = bInIsNavigationWindow;
 	OnTutorialClosed = InOnTutorialClosed;
 	OnTutorialExited = InOnTutorialExited;
 
@@ -100,12 +100,12 @@ void SEditorTutorials::HideContent()
 {
 	HandleHomeClicked();
 	bBrowserVisible = false;
-	bShowNavigation = false;
+	bIsNavigationWindow = false;
 }
 
 bool SEditorTutorials::IsNavigationVisible() const
 {
-	return bShowNavigation;
+	return bIsNavigationWindow;
 }
 
 EVisibility SEditorTutorials::GetBrowserVisibility() const
@@ -195,7 +195,7 @@ void SEditorTutorials::RebuildCurrentContent()
 			.OnClosed(FSimpleDelegate::CreateSP(this, &SEditorTutorials::HandleCloseClicked))
 			.IsStandalone(CurrentTutorial->bIsStandalone)
 			.ParentWindow(ParentWindow)
-			.AllowNonWidgetContent(bShowNavigation)
+			.AllowNonWidgetContent(bIsNavigationWindow)
 			.OnBackClicked(FSimpleDelegate::CreateSP(this, &SEditorTutorials::HandleBackClicked))
 			.OnHomeClicked(FSimpleDelegate::CreateSP(this, &SEditorTutorials::HandleHomeClicked))
 			.OnNextClicked(FSimpleDelegate::CreateSP(this, &SEditorTutorials::HandleNextClicked))
diff --git a/Engine/Source/Editor/IntroTutorials/Private/SEditorTutorials.h b/Engine/Source/Editor/IntroTutorials/Private/SEditorTutorials.h
--- a/Engine/Source/Editor/IntroTutorials/Private/SEditorTutorials.h
+++ b/Engine/Source/Editor/IntroTutorials/Private/SEditorTutorials.h
@@ -108,7 +108,7 @@ private:
 	bool bBrowserVisible;
 
 	/** Whether we should display navigation */
-	bool bShowNavigation;
+	bool bIsNavigationWindow;
 
 	/** Delegate fired when next button is clicked */
 	FOnNextClicked OnNextClicked;
diff --git a/Engine/Source/Editor/IntroTutorials/Private/STutorialContent.cpp b/Engine/Source/Editor/IntroTutorials/Private/STutorialContent.cpp
--- a/Engine/Source/Editor/IntroTutorials/Private/STutorialContent.cpp
+++ b/Engine/Source/Editor/IntroTutorials/Private/STutorialContent.cpp
@@ -42,6 +42,7 @@ void STutorialContent::Construct(const FArguments& InArgs, UEditorTutorial* InTu
 	IsHomeEnabled = InArgs._IsHomeEnabled;
 	IsNextEnabled = InArgs._IsNextEnabled;
 	Anchor = InArgs._Anchor;
+	bAllowNonWidgetContent = InArgs._AllowNonWidgetContent;
 
 	BorderIntroAnimation.AddCurve(0.0f, TutorialConstants::BorderIntroAnimationLength, ECurveEaseFunction::CubicOut);
 	BorderPulseAnimation.AddCurve(0.0f, TutorialConstants::BorderPulseAnimationLength, ECurveEaseFunction::Linear);
@@ -335,7 +336,7 @@ TSharedRef<SWidget> STutorialContent::GenerateContentWidget(const FTutorialConte
 
 FVector2D STutorialContent::GetPosition() const
 {
-	const bool bNonVisibleWidgetBound = !bIsVisible && Anchor.Type == ETutorialAnchorIdentifier::NamedWidget;
+	const bool bNonVisibleWidgetBound = bAllowNonWidgetContent && !bIsVisible && Anchor.Type == ETutorialAnchorIdentifier::NamedWidget;
 	if(bNonVisibleWidgetBound)
 	{
 		// fallback: center on cached window
@@ -444,7 +445,7 @@ EVisibility STutorialContent::GetVisibility() const
 	const bool bNonWidgetBound =  Anchor.Type == ETutorialAnchorIdentifier::None;
 
 	// fallback if widget is not drawn - we should display this content anyway
-	const bool bNonVisibleWidgetBound = !bIsVisible && Anchor.Type == ETutorialAnchorIdentifier::NamedWidget;
+	const bool bNonVisibleWidgetBound = bAllowNonWidgetContent && !bIsVisible && Anchor.Type == ETutorialAnchorIdentifier::NamedWidget;
 
 	return (bVisibleWidgetBound || bNonWidgetBound || bNonVisibleWidgetBound) ? EVisibility::SelfHitTestInvisible : EVisibility::Collapsed;
 }
diff --git a/Engine/Source/Editor/IntroTutorials/Private/STutorialContent.h b/Engine/Source/Editor/IntroTutorials/Private/STutorialContent.h
--- a/Engine/Source/Editor/IntroTutorials/Private/STutorialContent.h
+++ b/Engine/Source/Editor/IntroTutorials/Private/STutorialContent.h
@@ -56,6 +56,9 @@ class STutorialContent : public SCompoundWidget
 	/** Anchor if required */
 	SLATE_ARGUMENT(FTutorialContentAnchor, Anchor)
 
+	/** Whether we can show full window content in this overlay (i.e. in the same window as the navigation controls) */
+	SLATE_ARGUMENT(bool, AllowNonWidgetContent)
+
 	SLATE_END_ARGS()
 
 	void Construct(const FArguments& InArgs, UEditorTutorial* InTutorial, const FTutorialContent& InContent);
@@ -203,4 +206,7 @@ private:
 
 	/** Next button widget */
 	TSharedPtr<SWidget> NextButton;
+
+	/** Whether we can show full window content in this overlay (i.e. in the same window as the navigation controls) */
+	bool bAllowNonWidgetContent;
 };
\ No newline at end of file
diff --git a/Engine/Source/Editor/IntroTutorials/Private/STutorialOverlay.cpp b/Engine/Source/Editor/IntroTutorials/Private/STutorialOverlay.cpp
--- a/Engine/Source/Editor/IntroTutorials/Private/STutorialOverlay.cpp
+++ b/Engine/Source/Editor/IntroTutorials/Private/STutorialOverlay.cpp
@@ -79,7 +79,8 @@ void STutorialOverlay::Construct(const FArguments& InArgs, UEditorTutorial* InTu
 						.IsHomeEnabled(InArgs._IsHomeEnabled)
 						.IsNextEnabled(InArgs._IsNextEnabled)
 						.WrapTextAt(WidgetContent.ContentWidth)
-						.Anchor(WidgetContent.WidgetAnchor);
+						.Anchor(WidgetContent.WidgetAnchor)
+						.AllowNonWidgetContent(InArgs._AllowNonWidgetContent);
 					
 					PerformWidgetInteractions(WidgetContent); 					
 
