Fixed crash on shutdown that for some reason was only happening on Mac.

diff --git a/Engine/Source/Editor/IntroTutorials/Private/IntroTutorials.cpp b/Engine/Source/Editor/IntroTutorials/Private/IntroTutorials.cpp
--- a/Engine/Source/Editor/IntroTutorials/Private/IntroTutorials.cpp
+++ b/Engine/Source/Editor/IntroTutorials/Private/IntroTutorials.cpp
@@ -137,10 +137,11 @@ void FIntroTutorials::StartupModule()
 	PropertyEditorModule.RegisterCustomPropertyTypeLayout("TutorialContentAnchor", FOnGetPropertyTypeCustomizationInstance::CreateStatic(&FTutorialContentAnchorCustomization::MakeInstance));
 	PropertyEditorModule.RegisterCustomClassLayout("EditorTutorial", FOnGetDetailCustomizationInstance::CreateStatic(&FEditorTutorialDetailsCustomization::MakeInstance));
 
-	ContentIntroCurve = LoadObject<UCurveFloat>(nullptr, TEXT("/Engine/Tutorial/ContentIntroCurve.ContentIntroCurve"));
-	if(ContentIntroCurve)
+	UCurveFloat* ContentIntroCurveAsset = LoadObject<UCurveFloat>(nullptr, TEXT("/Engine/Tutorial/ContentIntroCurve.ContentIntroCurve"));
+	if(ContentIntroCurveAsset)
 	{
-		ContentIntroCurve->AddToRoot();
+		ContentIntroCurveAsset->AddToRoot();
+		ContentIntroCurve = ContentIntroCurveAsset;
 	}
 }
 
@@ -182,9 +183,9 @@ void FIntroTutorials::ShutdownModule()
 		PropertyEditorModule.UnregisterCustomClassLayout("EditorTutorial");
 	}
 
-	if(ContentIntroCurve)
+	if(ContentIntroCurve.IsValid())
 	{
-		ContentIntroCurve->RemoveFromRoot();
+		ContentIntroCurve.Get()->RemoveFromRoot();
 		ContentIntroCurve = nullptr;
 	}
 }
@@ -419,9 +420,9 @@ TSharedRef<SWidget> FIntroTutorials::CreateTutorialsWidget(FName InContext, TWea
 
 float FIntroTutorials::GetIntroCurveValue(float InTime)
 {
-	if(ContentIntroCurve)
+	if(ContentIntroCurve.IsValid())
 	{
-		return ContentIntroCurve->GetFloatValue(InTime);
+		return ContentIntroCurve.Get()->GetFloatValue(InTime);
 	}
 
 	return 1.0f;
diff --git a/Engine/Source/Editor/IntroTutorials/Private/IntroTutorials.h b/Engine/Source/Editor/IntroTutorials/Private/IntroTutorials.h
--- a/Engine/Source/Editor/IntroTutorials/Private/IntroTutorials.h
+++ b/Engine/Source/Editor/IntroTutorials/Private/IntroTutorials.h
@@ -113,5 +113,5 @@ private:
 	TSharedPtr<STutorialRoot> TutorialRoot;
 
 	/** Curve asset for intros */
-	UCurveFloat* ContentIntroCurve;
+	TWeakObjectPtr<UCurveFloat> ContentIntroCurve;
 };
\ No newline at end of file
