Fix tutorial not displaying if NULL window was passed in

diff --git a/Engine/Source/Editor/IntroTutorials/Private/IntroTutorials.cpp b/Engine/Source/Editor/IntroTutorials/Private/IntroTutorials.cpp
--- a/Engine/Source/Editor/IntroTutorials/Private/IntroTutorials.cpp
+++ b/Engine/Source/Editor/IntroTutorials/Private/IntroTutorials.cpp
@@ -356,7 +356,7 @@ FOnIsPicking& FIntroTutorials::OnIsPicking()
 	return OnIsPickingDelegate;
 }
 
-void FIntroTutorials::SummonTutorialBrowser(TSharedRef<SWindow> InWindow, const FString& InFilter)
+void FIntroTutorials::SummonTutorialBrowser(TWeakPtr<SWindow> InWindow, const FString& InFilter)
 {
 	if(TutorialRoot.IsValid())
 	{
diff --git a/Engine/Source/Editor/IntroTutorials/Private/IntroTutorials.h b/Engine/Source/Editor/IntroTutorials/Private/IntroTutorials.h
--- a/Engine/Source/Editor/IntroTutorials/Private/IntroTutorials.h
+++ b/Engine/Source/Editor/IntroTutorials/Private/IntroTutorials.h
@@ -24,7 +24,7 @@ public:
 
 	float GetIntroCurveValue(float InTime);
 
-	void SummonTutorialBrowser(TSharedRef<SWindow> InWindow, const FString& InFilter = TEXT(""));
+	void SummonTutorialBrowser(TWeakPtr<SWindow> InWindow = nullptr, const FString& InFilter = TEXT(""));
 
 	static FString AnalyticsEventNameFromTutorial(const FString& BaseEventName, UEditorTutorial* Tutorial);
 
@@ -85,7 +85,7 @@ public:
 	// IIntroTutorials interface
 	virtual void LaunchTutorial(UEditorTutorial* Tutorial, bool bRestart = true, TWeakPtr<SWindow> InNavigationWindow = nullptr, FSimpleDelegate OnTutorialClosed = FSimpleDelegate(), FSimpleDelegate OnTutorialExited = FSimpleDelegate()) override;
 	virtual void CloseAllTutorialContent() override;
-	virtual TSharedRef<SWidget> CreateTutorialsWidget(FName InContext, TWeakPtr<SWindow> InContextWindow) const override;
+	virtual TSharedRef<SWidget> CreateTutorialsWidget(FName InContext, TWeakPtr<SWindow> InContextWindow = nullptr) const override;
 	// End of IIntroTutorials interface
 private:
 	/** The tab id of the tutorial tab */
diff --git a/Engine/Source/Editor/IntroTutorials/Private/STutorialButton.cpp b/Engine/Source/Editor/IntroTutorials/Private/STutorialButton.cpp
--- a/Engine/Source/Editor/IntroTutorials/Private/STutorialButton.cpp
+++ b/Engine/Source/Editor/IntroTutorials/Private/STutorialButton.cpp
@@ -132,37 +132,34 @@ int32 STutorialButton::OnPaint(const FPaintArgs& Args, const FGeometry& Allotted
 
 FReply STutorialButton::HandleButtonClicked()
 {
-	RefreshStatus();
-
-	if(ContextWindow.IsValid())
-	{
-		if( FEngineAnalytics::IsAvailable() )
-		{
-			TArray<FAnalyticsEventAttribute> EventAttributes;
-			EventAttributes.Add(FAnalyticsEventAttribute(TEXT("Context"), Context.ToString()));
-			EventAttributes.Add(FAnalyticsEventAttribute(TEXT("TimeSinceAlertStarted"), (AlertStartTime != 0.0f && ShouldShowAlert()) ? (FPlatformTime::Seconds() - AlertStartTime) : -1.0f));
-			EventAttributes.Add(FAnalyticsEventAttribute(TEXT("LaunchedBrowser"), ShouldLaunchBrowser()));
+	UEditorTutorial* AttractTutorial = nullptr;
+	UEditorTutorial* LaunchTutorial = nullptr;
+	FString BrowserFilter;
+	GetDefault<UEditorTutorialSettings>()->FindTutorialInfoForContext(Context, AttractTutorial, LaunchTutorial, BrowserFilter);
 
-			FEngineAnalytics::GetProvider().RecordEvent( TEXT("Rocket.Tutorials.ClickedContextButton"), EventAttributes );
-		}
+	bTutorialAvailable = (LaunchTutorial != nullptr);
+	bTutorialCompleted = (LaunchTutorial != nullptr) && GetDefault<UTutorialStateSettings>()->HaveCompletedTutorial(LaunchTutorial);
+	bTutorialDismissed = (AttractTutorial != nullptr) && GetDefault<UTutorialStateSettings>()->IsTutorialDismissed(AttractTutorial);
 
-		FIntroTutorials& IntroTutorials = FModuleManager::GetModuleChecked<FIntroTutorials>(TEXT("IntroTutorials"));
-		if(ShouldLaunchBrowser())
-		{
-			IntroTutorials.SummonTutorialBrowser(ContextWindow.Pin().ToSharedRef(), CachedBrowserFilter);
-		}
-		else if (CachedLaunchTutorial != nullptr)
-		{
-			auto Delegate = FSimpleDelegate::CreateSP(this, &STutorialButton::HandleTutorialExited);
+	if (FEngineAnalytics::IsAvailable())
+	{
+		TArray<FAnalyticsEventAttribute> EventAttributes;
+		EventAttributes.Add(FAnalyticsEventAttribute(TEXT("Context"), Context.ToString()));
+		EventAttributes.Add(FAnalyticsEventAttribute(TEXT("TimeSinceAlertStarted"), (AlertStartTime != 0.0f && ShouldShowAlert()) ? (FPlatformTime::Seconds() - AlertStartTime) : -1.0f));
+		EventAttributes.Add(FAnalyticsEventAttribute(TEXT("LaunchedBrowser"), ShouldLaunchBrowser()));
 
-			const bool bRestart = true;
-			IntroTutorials.LaunchTutorial(CachedLaunchTutorial, bRestart, ContextWindow, Delegate, Delegate);
+		FEngineAnalytics::GetProvider().RecordEvent(TEXT("Rocket.Tutorials.ClickedContextButton"), EventAttributes);
+	}
 
-			const bool bDismissAcrossSessions = true;
-			GetMutableDefault<UTutorialStateSettings>()->DismissTutorial(CachedLaunchTutorial, bDismissAcrossSessions);
-			GetMutableDefault<UTutorialStateSettings>()->SaveProgress();
-			bTutorialDismissed = true;
-		}
+	FIntroTutorials& IntroTutorials = FModuleManager::GetModuleChecked<FIntroTutorials>(TEXT("IntroTutorials"));
+	if (ShouldLaunchBrowser())
+	{
+		IntroTutorials.SummonTutorialBrowser(ContextWindow, BrowserFilter);
+	}
+	else if (LaunchTutorial != nullptr)
+	{
+		const bool bRestart = true;
+		IntroTutorials.LaunchTutorial(LaunchTutorial, bRestart, ContextWindow);
 	}
 
 	return FReply::Handled();
@@ -212,16 +209,16 @@ void STutorialButton::DismissAlert()
 {
 	RefreshStatus();
 
-	if( FEngineAnalytics::IsAvailable() )
-	{
-		TArray<FAnalyticsEventAttribute> EventAttributes;
-		EventAttributes.Add(FAnalyticsEventAttribute(TEXT("Context"), Context.ToString()));
-		EventAttributes.Add(FAnalyticsEventAttribute(TEXT("TimeSinceAlertStarted"), (AlertStartTime != 0.0f && ShouldShowAlert()) ? (FPlatformTime::Seconds() - AlertStartTime) : -1.0f));
+		if( FEngineAnalytics::IsAvailable() )
+		{
+			TArray<FAnalyticsEventAttribute> EventAttributes;
+			EventAttributes.Add(FAnalyticsEventAttribute(TEXT("Context"), Context.ToString()));
+			EventAttributes.Add(FAnalyticsEventAttribute(TEXT("TimeSinceAlertStarted"), (AlertStartTime != 0.0f && ShouldShowAlert()) ? (FPlatformTime::Seconds() - AlertStartTime) : -1.0f));
 
-		FEngineAnalytics::GetProvider().RecordEvent( TEXT("Rocket.Tutorials.DismissedTutorialAlert"), EventAttributes );
-	}
+			FEngineAnalytics::GetProvider().RecordEvent( TEXT("Rocket.Tutorials.DismissedTutorialAlert"), EventAttributes );
+		}
 
-	const bool bDismissAcrossSessions = true;
+		const bool bDismissAcrossSessions = true;
 	if (CachedAttractTutorial != nullptr)
 	{
 		GetMutableDefault<UTutorialStateSettings>()->DismissTutorial(CachedAttractTutorial, bDismissAcrossSessions);
@@ -230,12 +227,12 @@ void STutorialButton::DismissAlert()
 	{
 		GetMutableDefault<UTutorialStateSettings>()->DismissTutorial(CachedLaunchTutorial, bDismissAcrossSessions);
 	}
-	GetMutableDefault<UTutorialStateSettings>()->SaveProgress();
-	bTutorialDismissed = true;
+		GetMutableDefault<UTutorialStateSettings>()->SaveProgress();
+		bTutorialDismissed = true;
 
-	FIntroTutorials& IntroTutorials = FModuleManager::GetModuleChecked<FIntroTutorials>(TEXT("IntroTutorials"));
-	IntroTutorials.CloseAllTutorialContent();
-}
+		FIntroTutorials& IntroTutorials = FModuleManager::GetModuleChecked<FIntroTutorials>(TEXT("IntroTutorials"));
+		IntroTutorials.CloseAllTutorialContent();
+	}
 
 void STutorialButton::LaunchTutorial()
 {
diff --git a/Engine/Source/Editor/IntroTutorials/Private/STutorialRoot.cpp b/Engine/Source/Editor/IntroTutorials/Private/STutorialRoot.cpp
--- a/Engine/Source/Editor/IntroTutorials/Private/STutorialRoot.cpp
+++ b/Engine/Source/Editor/IntroTutorials/Private/STutorialRoot.cpp
@@ -11,6 +11,7 @@
 #include "IToolkitHost.h"
 #include "EngineAnalytics.h"
 #include "Runtime/Analytics/Analytics/Public/Interfaces/IAnalyticsProvider.h"
+#include "Editor/MainFrame/Public/Interfaces/IMainFrameModule.h"
 
 #define LOCTEXT_NAMESPACE "STutorialRoot"
 
@@ -76,8 +77,15 @@ void STutorialRoot::Tick(const FGeometry& AllottedGeometry, const double InCurre
 	}
 }
 
-void STutorialRoot::SummonTutorialBrowser(TSharedRef<SWindow> InWindow, const FString& InFilter)
+void STutorialRoot::SummonTutorialBrowser(TWeakPtr<SWindow> InWindow, const FString& InFilter)
 {
+	// Use main frame if non specified
+	if (!InWindow.IsValid())
+	{
+		IMainFrameModule& MainFrameModule = FModuleManager::LoadModuleChecked<IMainFrameModule>(TEXT("MainFrame"));
+		InWindow = MainFrameModule.GetParentWindow();
+	}
+
 	TWeakPtr<SEditorTutorials>* FoundWidget = TutorialWidgets.Find(InWindow);
 	if(FoundWidget != nullptr && FoundWidget->IsValid())
 	{
diff --git a/Engine/Source/Editor/IntroTutorials/Private/STutorialRoot.h b/Engine/Source/Editor/IntroTutorials/Private/STutorialRoot.h
--- a/Engine/Source/Editor/IntroTutorials/Private/STutorialRoot.h
+++ b/Engine/Source/Editor/IntroTutorials/Private/STutorialRoot.h
@@ -34,7 +34,7 @@ public:
 	void CloseAllTutorialContent();
 
 	/** Summon the browser widget for the specified window */
-	void SummonTutorialBrowser(TSharedRef<SWindow> InWindow, const FString& InFilter = TEXT(""));
+	void SummonTutorialBrowser(TWeakPtr<SWindow> InWindow = nullptr, const FString& InFilter = TEXT(""));
 
 	/** Reload tutorials that we know about */
 	void ReloadTutorials();
diff --git a/Engine/Source/Editor/IntroTutorials/Public/IIntroTutorials.h b/Engine/Source/Editor/IntroTutorials/Public/IIntroTutorials.h
--- a/Engine/Source/Editor/IntroTutorials/Public/IIntroTutorials.h
+++ b/Engine/Source/Editor/IntroTutorials/Public/IIntroTutorials.h
@@ -56,6 +56,6 @@ public:
 	 * @param	InContextWindow		The window that the context is attached to (e.g. the main window, or an asset editor tab)
 	 * @return a widget used to access context-sensitive tutorials
 	 */
-	virtual TSharedRef<SWidget> CreateTutorialsWidget(FName InContext, TWeakPtr<SWindow> InContextWindow) const = 0;
+	virtual TSharedRef<SWidget> CreateTutorialsWidget(FName InContext, TWeakPtr<SWindow> InContextWindow = nullptr) const = 0;
 };
 
