Sorted tutorials browser entries

So the list is stable, for now we do a simple alphabetical sort.

diff --git a/Engine/Source/Editor/IntroTutorials/Private/STutorialsBrowser.cpp b/Engine/Source/Editor/IntroTutorials/Private/STutorialsBrowser.cpp
--- a/Engine/Source/Editor/IntroTutorials/Private/STutorialsBrowser.cpp
+++ b/Engine/Source/Editor/IntroTutorials/Private/STutorialsBrowser.cpp
@@ -150,6 +150,17 @@ public:
 		return bPassesFilter && bPassesCategory;
 	}
 
+	FString GetTitleString() const override
+	{
+		const FString Title = !Category.Title.IsEmpty() ? Category.Title.ToString() : CategoryName;
+		return Title;
+	}
+
+	bool SortAgainst(TSharedRef<ITutorialListEntry> OtherEntry) const override
+	{
+		return GetTitleString() < OtherEntry->GetTitleString();
+	}
+
 	void AddSubCategory(TSharedPtr<FTutorialListEntry_Category> InSubCategory)
 	{
 		SubCategories.Add(InSubCategory);
@@ -348,6 +359,16 @@ public:
 		return bPassesFilter && bPassesCategory;
 	}
 
+	FString GetTitleString() const override
+	{
+		return Tutorial->Title.ToString();
+	}
+
+	bool SortAgainst(TSharedRef<ITutorialListEntry> OtherEntry) const override
+	{
+		return GetTitleString() < OtherEntry->GetTitleString();
+	}
+
 	FReply OnClicked(bool bRestart) const
 	{
 		OnTutorialSelected.ExecuteIfBound(Tutorial, bRestart);
@@ -768,6 +789,17 @@ void STutorialsBrowser::FilterTutorials()
 		CurrentCategoryPtr = CurrentCategory;
 	}
 
+	FilteredEntries.Sort(
+		[](TSharedPtr<ITutorialListEntry> EntryA, TSharedPtr<ITutorialListEntry> EntryB)->bool
+		{
+			if(EntryA.IsValid() && EntryB.IsValid())
+			{
+				return EntryA->SortAgainst(EntryB.ToSharedRef());
+			}
+			return false;
+		}
+	);
+
 	TutorialList->RequestListRefresh();
 }
 
diff --git a/Engine/Source/Editor/IntroTutorials/Private/STutorialsBrowser.h b/Engine/Source/Editor/IntroTutorials/Private/STutorialsBrowser.h
--- a/Engine/Source/Editor/IntroTutorials/Private/STutorialsBrowser.h
+++ b/Engine/Source/Editor/IntroTutorials/Private/STutorialsBrowser.h
@@ -14,6 +14,15 @@ struct ITutorialListEntry
 
 	/** Whether this entry passes the current filter criteria */
 	virtual bool PassesFilter(const FString& InCategoryFilter, const FString& InFilter) const = 0;
+
+	/** Get the string representation of this item's title */
+	virtual FString GetTitleString() const = 0;
+
+	/** 
+	 * Sort this entry against another entry
+	 * @return true if this < OtherEntry 
+	 */
+	virtual bool SortAgainst(TSharedRef<ITutorialListEntry> OtherEntry) const = 0;
 };
 
 class FTutorialListEntry_Category;
