Tutorial alert now dismisses again when tutorial is taken

Undo issue caused by bad merge

diff --git a/Engine/Source/Editor/IntroTutorials/Private/STutorialButton.cpp b/Engine/Source/Editor/IntroTutorials/Private/STutorialButton.cpp
--- a/Engine/Source/Editor/IntroTutorials/Private/STutorialButton.cpp
+++ b/Engine/Source/Editor/IntroTutorials/Private/STutorialButton.cpp
@@ -132,34 +132,34 @@ int32 STutorialButton::OnPaint(const FPaintArgs& Args, const FGeometry& Allotted
 
 FReply STutorialButton::HandleButtonClicked()
 {
-	UEditorTutorial* AttractTutorial = nullptr;
-	UEditorTutorial* LaunchTutorial = nullptr;
-	FString BrowserFilter;
-	GetDefault<UEditorTutorialSettings>()->FindTutorialInfoForContext(Context, AttractTutorial, LaunchTutorial, BrowserFilter);
-
-	bTutorialAvailable = (LaunchTutorial != nullptr);
-	bTutorialCompleted = (LaunchTutorial != nullptr) && GetDefault<UTutorialStateSettings>()->HaveCompletedTutorial(LaunchTutorial);
-	bTutorialDismissed = (AttractTutorial != nullptr) && GetDefault<UTutorialStateSettings>()->IsTutorialDismissed(AttractTutorial);
+	RefreshStatus();
 
-	if (FEngineAnalytics::IsAvailable())
+	if( FEngineAnalytics::IsAvailable() )
 	{
 		TArray<FAnalyticsEventAttribute> EventAttributes;
 		EventAttributes.Add(FAnalyticsEventAttribute(TEXT("Context"), Context.ToString()));
 		EventAttributes.Add(FAnalyticsEventAttribute(TEXT("TimeSinceAlertStarted"), (AlertStartTime != 0.0f && ShouldShowAlert()) ? (FPlatformTime::Seconds() - AlertStartTime) : -1.0f));
 		EventAttributes.Add(FAnalyticsEventAttribute(TEXT("LaunchedBrowser"), ShouldLaunchBrowser()));
 
-		FEngineAnalytics::GetProvider().RecordEvent(TEXT("Rocket.Tutorials.ClickedContextButton"), EventAttributes);
+		FEngineAnalytics::GetProvider().RecordEvent( TEXT("Rocket.Tutorials.ClickedContextButton"), EventAttributes );
 	}
 
 	FIntroTutorials& IntroTutorials = FModuleManager::GetModuleChecked<FIntroTutorials>(TEXT("IntroTutorials"));
-	if (ShouldLaunchBrowser())
+	if(ShouldLaunchBrowser())
 	{
-		IntroTutorials.SummonTutorialBrowser(ContextWindow, BrowserFilter);
+		IntroTutorials.SummonTutorialBrowser(ContextWindow.Pin().ToSharedRef(), CachedBrowserFilter);
 	}
-	else if (LaunchTutorial != nullptr)
+	else if (CachedLaunchTutorial != nullptr)
 	{
+		auto Delegate = FSimpleDelegate::CreateSP(this, &STutorialButton::HandleTutorialExited);
+
 		const bool bRestart = true;
-		IntroTutorials.LaunchTutorial(LaunchTutorial, bRestart, ContextWindow);
+		IntroTutorials.LaunchTutorial(CachedLaunchTutorial, bRestart, ContextWindow, Delegate, Delegate);
+
+		const bool bDismissAcrossSessions = true;
+		GetMutableDefault<UTutorialStateSettings>()->DismissTutorial(CachedLaunchTutorial, bDismissAcrossSessions);
+		GetMutableDefault<UTutorialStateSettings>()->SaveProgress();
+		bTutorialDismissed = true;
 	}
 
 	return FReply::Handled();
@@ -209,16 +209,16 @@ void STutorialButton::DismissAlert()
 {
 	RefreshStatus();
 
-		if( FEngineAnalytics::IsAvailable() )
-		{
-			TArray<FAnalyticsEventAttribute> EventAttributes;
-			EventAttributes.Add(FAnalyticsEventAttribute(TEXT("Context"), Context.ToString()));
-			EventAttributes.Add(FAnalyticsEventAttribute(TEXT("TimeSinceAlertStarted"), (AlertStartTime != 0.0f && ShouldShowAlert()) ? (FPlatformTime::Seconds() - AlertStartTime) : -1.0f));
+	if( FEngineAnalytics::IsAvailable() )
+	{
+		TArray<FAnalyticsEventAttribute> EventAttributes;
+		EventAttributes.Add(FAnalyticsEventAttribute(TEXT("Context"), Context.ToString()));
+		EventAttributes.Add(FAnalyticsEventAttribute(TEXT("TimeSinceAlertStarted"), (AlertStartTime != 0.0f && ShouldShowAlert()) ? (FPlatformTime::Seconds() - AlertStartTime) : -1.0f));
 
-			FEngineAnalytics::GetProvider().RecordEvent( TEXT("Rocket.Tutorials.DismissedTutorialAlert"), EventAttributes );
-		}
+		FEngineAnalytics::GetProvider().RecordEvent( TEXT("Rocket.Tutorials.DismissedTutorialAlert"), EventAttributes );
+	}
 
-		const bool bDismissAcrossSessions = true;
+	const bool bDismissAcrossSessions = true;
 	if (CachedAttractTutorial != nullptr)
 	{
 		GetMutableDefault<UTutorialStateSettings>()->DismissTutorial(CachedAttractTutorial, bDismissAcrossSessions);
@@ -227,12 +227,12 @@ void STutorialButton::DismissAlert()
 	{
 		GetMutableDefault<UTutorialStateSettings>()->DismissTutorial(CachedLaunchTutorial, bDismissAcrossSessions);
 	}
-		GetMutableDefault<UTutorialStateSettings>()->SaveProgress();
-		bTutorialDismissed = true;
+	GetMutableDefault<UTutorialStateSettings>()->SaveProgress();
+	bTutorialDismissed = true;
 
-		FIntroTutorials& IntroTutorials = FModuleManager::GetModuleChecked<FIntroTutorials>(TEXT("IntroTutorials"));
-		IntroTutorials.CloseAllTutorialContent();
-	}
+	FIntroTutorials& IntroTutorials = FModuleManager::GetModuleChecked<FIntroTutorials>(TEXT("IntroTutorials"));
+	IntroTutorials.CloseAllTutorialContent();
+}
 
 void STutorialButton::LaunchTutorial()
 {
