Correctly open tutorial when opening an asset for a tutorial.

Make sure we have an overlay constructed for the window we just created for the asset.

TTP# 347174 - TUTORIAL: When clicking on a sub editor tutorial, the tutorial does not automatically popup if the sub editor is not already opened

diff --git a/Engine/Source/Editor/IntroTutorials/Private/STutorialRoot.cpp b/Engine/Source/Editor/IntroTutorials/Private/STutorialRoot.cpp
--- a/Engine/Source/Editor/IntroTutorials/Private/STutorialRoot.cpp
+++ b/Engine/Source/Editor/IntroTutorials/Private/STutorialRoot.cpp
@@ -116,6 +116,12 @@ void STutorialRoot::LaunchTutorial(UEditorTutorial* InTutorial, bool bInRestart,
 				if(Toolkit.IsValid())
 				{
 					InNavigationWindow = FSlateApplication::Get().FindWidgetWindow(Toolkit->GetToolkitHost()->GetParentWidget());
+
+					// make sure we have a valid tutorial overlay
+					if(InNavigationWindow.IsValid())
+					{
+						MaybeAddOverlay(InNavigationWindow.Pin().ToSharedRef());
+					}
 				}
 			}
 		}
