Dont show non-standard Blueprint Tutorials

Also allow tutorials to be hidden from the browser if necessary

TTP# 347111 - TUTORIALS:Tutorials browser shows MacroLibraries based on EditorTutorial class

diff --git a/Engine/Source/Editor/IntroTutorials/Private/STutorialsBrowser.cpp b/Engine/Source/Editor/IntroTutorials/Private/STutorialsBrowser.cpp
--- a/Engine/Source/Editor/IntroTutorials/Private/STutorialsBrowser.cpp
+++ b/Engine/Source/Editor/IntroTutorials/Private/STutorialsBrowser.cpp
@@ -633,9 +633,13 @@ void STutorialsBrowser::RebuildTutorials(TSharedPtr<FTutorialListEntry_Category>
 	for (const auto& TutorialAsset : AssetData)
 	{
 		UBlueprint* Blueprint = LoadObject<UBlueprint>(nullptr, *TutorialAsset.ObjectPath.ToString());
-		if (Blueprint && Blueprint->GeneratedClass)
+		if (Blueprint && Blueprint->GeneratedClass && Blueprint->BlueprintType == BPTYPE_Normal)
 		{
-			Tutorials.Add(MakeShareable(new FTutorialListEntry_Tutorial(Blueprint->GeneratedClass->GetDefaultObject<UEditorTutorial>(), FOnTutorialSelected::CreateSP(this, &STutorialsBrowser::OnTutorialSelected), TAttribute<FText>::Create(TAttribute<FText>::FGetter::CreateSP(this, &STutorialsBrowser::GetSearchText)))));
+			UEditorTutorial* Tutorial = Blueprint->GeneratedClass->GetDefaultObject<UEditorTutorial>();
+			if(!Tutorial->bHideInBrowser)
+			{
+				Tutorials.Add(MakeShareable(new FTutorialListEntry_Tutorial(Tutorial, FOnTutorialSelected::CreateSP(this, &STutorialsBrowser::OnTutorialSelected), TAttribute<FText>::Create(TAttribute<FText>::FGetter::CreateSP(this, &STutorialsBrowser::GetSearchText)))));
+			}
 		}
 	}
 
diff --git a/Engine/Source/Editor/IntroTutorials/Public/EditorTutorial.h b/Engine/Source/Editor/IntroTutorials/Public/EditorTutorial.h
--- a/Engine/Source/Editor/IntroTutorials/Public/EditorTutorial.h
+++ b/Engine/Source/Editor/IntroTutorials/Public/EditorTutorial.h
@@ -244,6 +244,10 @@ class INTROTUTORIALS_API UEditorTutorial : public UObject
 	UPROPERTY()
 	FString ImportPath;
 
+	/** Hide this tutorial in the tutorials browser */
+	UPROPERTY(EditAnywhere, Category="Tutorial")
+	bool bHideInBrowser;
+
 public:
 	/** Called when a tutorial stage is started */
 	void HandleTutorialStageStarted(FName StageName);
