Internal builds default to not showing tutorial alert pulse

This can be enbled internally by using the -TestTutorialAlerts command line switch.

diff --git a/Engine/Source/Editor/IntroTutorials/Private/IntroTutorials.cpp b/Engine/Source/Editor/IntroTutorials/Private/IntroTutorials.cpp
--- a/Engine/Source/Editor/IntroTutorials/Private/IntroTutorials.cpp
+++ b/Engine/Source/Editor/IntroTutorials/Private/IntroTutorials.cpp
@@ -30,6 +30,7 @@
 #include "STutorialButton.h"
 #include "ToolkitManager.h"
 #include "BlueprintEditor.h"
+#include "EngineBuildSettings.h"
 
 #define LOCTEXT_NAMESPACE "IntroTutorials"
 
@@ -286,27 +287,30 @@ void FIntroTutorials::SummonBlueprintTutorialHome(UObject* Asset, bool bForceWel
 
 bool FIntroTutorials::MaybeOpenWelcomeTutorial()
 {
-	// try editor startup tutorial
-	TSubclassOf<UEditorTutorial> EditorStartupTutorialClass = LoadClass<UEditorTutorial>(NULL, *GetDefault<UEditorTutorialSettings>()->StartupTutorial.AssetLongPathname, NULL, LOAD_None, NULL);
-	if(EditorStartupTutorialClass != nullptr)
+	if(FParse::Param(FCommandLine::Get(), TEXT("TestTutorialAlerts")) || !FEngineBuildSettings::IsInternalBuild())
 	{
-		UEditorTutorial* Tutorial = EditorStartupTutorialClass->GetDefaultObject<UEditorTutorial>();
-		if (!GetDefault<UTutorialStateSettings>()->HaveSeenTutorial(Tutorial))
+		// try editor startup tutorial
+		TSubclassOf<UEditorTutorial> EditorStartupTutorialClass = LoadClass<UEditorTutorial>(NULL, *GetDefault<UEditorTutorialSettings>()->StartupTutorial.AssetLongPathname, NULL, LOAD_None, NULL);
+		if(EditorStartupTutorialClass != nullptr)
 		{
-			LaunchTutorial(Tutorial);
-			return true;
+			UEditorTutorial* Tutorial = EditorStartupTutorialClass->GetDefaultObject<UEditorTutorial>();
+			if (!GetDefault<UTutorialStateSettings>()->HaveSeenTutorial(Tutorial))
+			{
+				LaunchTutorial(Tutorial);
+				return true;
+			}
 		}
-	}
 
-	// Try project startup tutorial
-	TSubclassOf<UEditorTutorial> ProjectStartupTutorialClass = LoadClass<UEditorTutorial>(NULL, *GetDefault<UTutorialSettings>()->StartupTutorial.AssetLongPathname, NULL, LOAD_None, NULL);
-	if(ProjectStartupTutorialClass != nullptr)
-	{
-		UEditorTutorial* Tutorial = ProjectStartupTutorialClass->GetDefaultObject<UEditorTutorial>();
-		if (!GetDefault<UTutorialStateSettings>()->HaveSeenTutorial(Tutorial))
+		// Try project startup tutorial
+		TSubclassOf<UEditorTutorial> ProjectStartupTutorialClass = LoadClass<UEditorTutorial>(NULL, *GetDefault<UTutorialSettings>()->StartupTutorial.AssetLongPathname, NULL, LOAD_None, NULL);
+		if(ProjectStartupTutorialClass != nullptr)
 		{
-			LaunchTutorial(Tutorial);
-			return true;
+			UEditorTutorial* Tutorial = ProjectStartupTutorialClass->GetDefaultObject<UEditorTutorial>();
+			if (!GetDefault<UTutorialStateSettings>()->HaveSeenTutorial(Tutorial))
+			{
+				LaunchTutorial(Tutorial);
+				return true;
+			}
 		}
 	}
 
diff --git a/Engine/Source/Editor/IntroTutorials/Private/STutorialButton.cpp b/Engine/Source/Editor/IntroTutorials/Private/STutorialButton.cpp
--- a/Engine/Source/Editor/IntroTutorials/Private/STutorialButton.cpp
+++ b/Engine/Source/Editor/IntroTutorials/Private/STutorialButton.cpp
@@ -5,6 +5,7 @@
 #include "EditorTutorialSettings.h"
 #include "TutorialStateSettings.h"
 #include "TutorialMetaData.h"
+#include "EngineBuildSettings.h"
 
 #define LOCTEXT_NAMESPACE "STutorialButton"
 
@@ -19,6 +20,8 @@ void STutorialButton::Construct(const FArguments& InArgs)
 	Context = InArgs._Context;
 	ContextWindow = InArgs._ContextWindow;
 
+	bTestAlerts = FParse::Param(FCommandLine::Get(), TEXT("TestTutorialAlerts"));
+
 	bTutorialAvailable = false;
 	bTutorialCompleted = false;
 	bTutorialDismissed = false;
@@ -257,7 +260,7 @@ bool STutorialButton::ShouldLaunchBrowser() const
 
 bool STutorialButton::ShouldShowAlert() const
 {
-	return (bTutorialAvailable && !(bTutorialCompleted || bTutorialDismissed));
+	return ((bTestAlerts || !FEngineBuildSettings::IsInternalBuild()) && bTutorialAvailable && !(bTutorialCompleted || bTutorialDismissed));
 }
 
 FText STutorialButton::GetButtonToolTip() const
diff --git a/Engine/Source/Editor/IntroTutorials/Private/STutorialButton.h b/Engine/Source/Editor/IntroTutorials/Private/STutorialButton.h
--- a/Engine/Source/Editor/IntroTutorials/Private/STutorialButton.h
+++ b/Engine/Source/Editor/IntroTutorials/Private/STutorialButton.h
@@ -60,6 +60,9 @@ private:
 	/** Whether we have dismissed the tutorial for this content */
 	bool bTutorialDismissed;
 
+	/** Flag to force alerts to appear in internal builds (caches command line -TestTutorialAlerts) */
+	bool bTestAlerts;
+
 	/** Context that this widget was created for (i.e. what part of the editor) */
 	FName Context;
 
