Tearing off a highlighted widget no longer results in duplicate content being displayed

Tutorials browser now always appears on the mainframe. This avoids the browser appearing on small torn-off windows if content is torn off while a tutorial is in progress.
Also fixed crash when opening a browser from a widget that used to be docked to anotehr tab (weak window pointer was invalid at this point).

TTP# 347014 - EDITOR: TUTORIAL: If a tab is undocked while being referenced in the "welcome to the unreal editor" tutorial, there will be two instances of the tutorial text

diff --git a/Engine/Source/Editor/IntroTutorials/Private/IntroTutorials.cpp b/Engine/Source/Editor/IntroTutorials/Private/IntroTutorials.cpp
--- a/Engine/Source/Editor/IntroTutorials/Private/IntroTutorials.cpp
+++ b/Engine/Source/Editor/IntroTutorials/Private/IntroTutorials.cpp
@@ -249,8 +249,7 @@ void FIntroTutorials::MainFrameLoad(TSharedPtr<SWindow> InRootWindow, bool bIsNe
 
 void FIntroTutorials::SummonTutorialHome()
 {
-	FLevelEditorModule& LevelEditorModule = FModuleManager::LoadModuleChecked<FLevelEditorModule>( "LevelEditor" );
-	SummonTutorialBrowser(LevelEditorModule.GetLevelEditorTab()->GetParentWindow().ToSharedRef());
+	SummonTutorialBrowser();
 }
 
 void FIntroTutorials::SummonBlueprintTutorialHome(UObject* Asset, bool bForceWelcome)
diff --git a/Engine/Source/Editor/IntroTutorials/Private/SEditorTutorials.cpp b/Engine/Source/Editor/IntroTutorials/Private/SEditorTutorials.cpp
--- a/Engine/Source/Editor/IntroTutorials/Private/SEditorTutorials.cpp
+++ b/Engine/Source/Editor/IntroTutorials/Private/SEditorTutorials.cpp
@@ -20,6 +20,8 @@ void SEditorTutorials::Construct(const FArguments& InArgs)
 	OnCloseClicked = InArgs._OnCloseClicked;
 	OnGetCurrentTutorial = InArgs._OnGetCurrentTutorial;
 	OnGetCurrentTutorialStage = InArgs._OnGetCurrentTutorialStage;
+	OnWidgetWasDrawn = InArgs._OnWidgetWasDrawn;
+	OnWasWidgetDrawn = InArgs._OnWasWidgetDrawn;
 
 	TutorialHome = SNew(STutorialsBrowser)
 		.Visibility(this, &SEditorTutorials::GetBrowserVisibility)
@@ -93,7 +95,10 @@ void SEditorTutorials::LaunchTutorial(bool bInIsNavigationWindow, FSimpleDelegat
 void SEditorTutorials::ShowBrowser(const FString& InFilter)
 {
 	TutorialHome->SetFilter(InFilter);
-	HandleHomeClicked();
+	OnHomeClicked.ExecuteIfBound();
+	ContentBox->ClearChildren();
+	OnTutorialExited.ExecuteIfBound();
+	bBrowserVisible = true;
 }
 
 void SEditorTutorials::HideContent()
@@ -136,11 +141,11 @@ void SEditorTutorials::HandleBackClicked()
 void SEditorTutorials::HandleHomeClicked()
 {
 	OnHomeClicked.ExecuteIfBound();
-
 	ContentBox->ClearChildren();
-	bBrowserVisible = true;
-	TutorialHome->ReloadTutorials();
 	OnTutorialExited.ExecuteIfBound();
+
+	FIntroTutorials& IntroTutorials = FModuleManager::GetModuleChecked<FIntroTutorials>(TEXT("IntroTutorials"));
+	IntroTutorials.SummonTutorialBrowser();
 }
 
 void SEditorTutorials::HandleNextClicked()
@@ -202,6 +207,8 @@ void SEditorTutorials::RebuildCurrentContent()
 			.IsBackEnabled(this, &SEditorTutorials::IsBackButtonEnabled)
 			.IsHomeEnabled(this, &SEditorTutorials::IsHomeButtonEnabled)
 			.IsNextEnabled(this, &SEditorTutorials::IsNextButtonEnabled)
+			.OnWidgetWasDrawn(OnWidgetWasDrawn)
+			.OnWasWidgetDrawn(OnWasWidgetDrawn)
 		];
 	}
 	else
diff --git a/Engine/Source/Editor/IntroTutorials/Private/SEditorTutorials.h b/Engine/Source/Editor/IntroTutorials/Private/SEditorTutorials.h
--- a/Engine/Source/Editor/IntroTutorials/Private/SEditorTutorials.h
+++ b/Engine/Source/Editor/IntroTutorials/Private/SEditorTutorials.h
@@ -3,9 +3,7 @@
 #pragma once
 
 #include "STutorialsBrowser.h"
-
-class STutorialsBrowser;
-class STutorialOverlay;
+#include "STutorialOverlay.h"
 
 /** Delegate fired when next button is clicked */
 DECLARE_DELEGATE_OneParam(FOnNextClicked, TWeakPtr<SWindow> /*InNavigationWindow*/);
@@ -26,14 +24,16 @@ class SEditorTutorials : public SCompoundWidget
 		_Visibility = EVisibility::SelfHitTestInvisible;
 	}
 
-	SLATE_ARGUMENT(FSimpleDelegate, OnCloseClicked)
+	SLATE_EVENT(FSimpleDelegate, OnCloseClicked)
 	SLATE_ARGUMENT(TWeakPtr<SWindow>, ParentWindow)
-	SLATE_ARGUMENT(FOnNextClicked, OnNextClicked)
-	SLATE_ARGUMENT(FSimpleDelegate, OnBackClicked)
-	SLATE_ARGUMENT(FSimpleDelegate, OnHomeClicked)
-	SLATE_ARGUMENT(FOnGetCurrentTutorial, OnGetCurrentTutorial)
-	SLATE_ARGUMENT(FOnGetCurrentTutorialStage, OnGetCurrentTutorialStage)
-	SLATE_ARGUMENT(FOnLaunchTutorial, OnLaunchTutorial)
+	SLATE_EVENT(FOnNextClicked, OnNextClicked)
+	SLATE_EVENT(FSimpleDelegate, OnBackClicked)
+	SLATE_EVENT(FSimpleDelegate, OnHomeClicked)
+	SLATE_EVENT(FOnGetCurrentTutorial, OnGetCurrentTutorial)
+	SLATE_EVENT(FOnGetCurrentTutorialStage, OnGetCurrentTutorialStage)
+	SLATE_EVENT(FOnLaunchTutorial, OnLaunchTutorial)
+	SLATE_EVENT(FOnWidgetWasDrawn, OnWidgetWasDrawn)
+	SLATE_EVENT(FOnWasWidgetDrawn, OnWasWidgetDrawn)
 
 	SLATE_END_ARGS()
 
@@ -131,4 +131,8 @@ private:
 	/** External delegates used to report user interaction */
 	FSimpleDelegate OnTutorialClosed;
 	FSimpleDelegate OnTutorialExited;
+
+	/** Delegates for registering & querying whether a widget was drawn */
+	FOnWidgetWasDrawn OnWidgetWasDrawn;
+	FOnWasWidgetDrawn OnWasWidgetDrawn;
 };
\ No newline at end of file
diff --git a/Engine/Source/Editor/IntroTutorials/Private/STutorialButton.cpp b/Engine/Source/Editor/IntroTutorials/Private/STutorialButton.cpp
--- a/Engine/Source/Editor/IntroTutorials/Private/STutorialButton.cpp
+++ b/Engine/Source/Editor/IntroTutorials/Private/STutorialButton.cpp
@@ -150,7 +150,7 @@ FReply STutorialButton::HandleButtonClicked()
 	FIntroTutorials& IntroTutorials = FModuleManager::GetModuleChecked<FIntroTutorials>(TEXT("IntroTutorials"));
 	if(ShouldLaunchBrowser())
 	{
-		IntroTutorials.SummonTutorialBrowser(ContextWindow.Pin().ToSharedRef(), CachedBrowserFilter);
+		IntroTutorials.SummonTutorialBrowser(nullptr, CachedBrowserFilter);
 	}
 	else if (CachedLaunchTutorial != nullptr)
 	{
@@ -244,13 +244,10 @@ void STutorialButton::LaunchTutorial()
 
 void STutorialButton::LaunchBrowser()
 {
-	if(ContextWindow.IsValid())
-	{
-		RefreshStatus();
+	RefreshStatus();
 
-		FIntroTutorials& IntroTutorials = FModuleManager::GetModuleChecked<FIntroTutorials>(TEXT("IntroTutorials"));
-		IntroTutorials.SummonTutorialBrowser(ContextWindow.Pin().ToSharedRef(), CachedBrowserFilter);
-	}
+	FIntroTutorials& IntroTutorials = FModuleManager::GetModuleChecked<FIntroTutorials>(TEXT("IntroTutorials"));
+	IntroTutorials.SummonTutorialBrowser(nullptr, CachedBrowserFilter);
 }
 
 bool STutorialButton::ShouldLaunchBrowser() const
diff --git a/Engine/Source/Editor/IntroTutorials/Private/STutorialContent.cpp b/Engine/Source/Editor/IntroTutorials/Private/STutorialContent.cpp
--- a/Engine/Source/Editor/IntroTutorials/Private/STutorialContent.cpp
+++ b/Engine/Source/Editor/IntroTutorials/Private/STutorialContent.cpp
@@ -43,6 +43,7 @@ void STutorialContent::Construct(const FArguments& InArgs, UEditorTutorial* InTu
 	IsNextEnabled = InArgs._IsNextEnabled;
 	Anchor = InArgs._Anchor;
 	bAllowNonWidgetContent = InArgs._AllowNonWidgetContent;
+	OnWasWidgetDrawn = InArgs._OnWasWidgetDrawn;
 
 	BorderIntroAnimation.AddCurve(0.0f, TutorialConstants::BorderIntroAnimationLength, ECurveEaseFunction::CubicOut);
 	BorderPulseAnimation.AddCurve(0.0f, TutorialConstants::BorderPulseAnimationLength, ECurveEaseFunction::Linear);
@@ -336,7 +337,15 @@ TSharedRef<SWidget> STutorialContent::GenerateContentWidget(const FTutorialConte
 
 FVector2D STutorialContent::GetPosition() const
 {
-	const bool bNonVisibleWidgetBound = bAllowNonWidgetContent && !bIsVisible && Anchor.Type == ETutorialAnchorIdentifier::NamedWidget;
+	bool bNonVisibleWidgetBound = bAllowNonWidgetContent && !bIsVisible && Anchor.Type == ETutorialAnchorIdentifier::NamedWidget;
+	if(bNonVisibleWidgetBound)
+	{
+		if(OnWasWidgetDrawn.IsBound())
+		{
+			bNonVisibleWidgetBound &= !OnWasWidgetDrawn.Execute(Anchor.WrapperIdentifier);
+		}
+	}
+
 	if(bNonVisibleWidgetBound)
 	{
 		// fallback: center on cached window
@@ -445,7 +454,14 @@ EVisibility STutorialContent::GetVisibility() const
 	const bool bNonWidgetBound =  Anchor.Type == ETutorialAnchorIdentifier::None;
 
 	// fallback if widget is not drawn - we should display this content anyway
-	const bool bNonVisibleWidgetBound = bAllowNonWidgetContent && !bIsVisible && Anchor.Type == ETutorialAnchorIdentifier::NamedWidget;
+	bool bNonVisibleWidgetBound = bAllowNonWidgetContent && !bIsVisible && Anchor.Type == ETutorialAnchorIdentifier::NamedWidget;
+	if(bNonVisibleWidgetBound)
+	{
+		if(OnWasWidgetDrawn.IsBound())
+		{
+			bNonVisibleWidgetBound &= !OnWasWidgetDrawn.Execute(Anchor.WrapperIdentifier);
+		}	
+	}
 
 	return (bVisibleWidgetBound || bNonWidgetBound || bNonVisibleWidgetBound) ? EVisibility::SelfHitTestInvisible : EVisibility::Collapsed;
 }
diff --git a/Engine/Source/Editor/IntroTutorials/Private/STutorialContent.h b/Engine/Source/Editor/IntroTutorials/Private/STutorialContent.h
--- a/Engine/Source/Editor/IntroTutorials/Private/STutorialContent.h
+++ b/Engine/Source/Editor/IntroTutorials/Private/STutorialContent.h
@@ -5,6 +5,7 @@
 #include "IDocumentation.h"
 #include "IDocumentationPage.h"
 #include "EditorTutorial.h"
+#include "STutorialOverlay.h"
 
 /**
  * The widget which displays 'floating' content 
@@ -30,16 +31,16 @@ class STutorialContent : public SCompoundWidget
 	SLATE_ARGUMENT(bool, IsStandalone)
 
 	/** Delegate fired when the close button is clicked */
-	SLATE_ARGUMENT(FSimpleDelegate, OnClosed)
+	SLATE_EVENT(FSimpleDelegate, OnClosed)
 
 	/** Delegate fired when the back button is clicked */
-	SLATE_ARGUMENT(FSimpleDelegate, OnBackClicked)
+	SLATE_EVENT(FSimpleDelegate, OnBackClicked)
 
 	/** Delegate fired when the home button is clicked */
-	SLATE_ARGUMENT(FSimpleDelegate, OnHomeClicked)
+	SLATE_EVENT(FSimpleDelegate, OnHomeClicked)
 
 	/** Delegate fired when the next button is clicked */
-	SLATE_ARGUMENT(FSimpleDelegate, OnNextClicked)
+	SLATE_EVENT(FSimpleDelegate, OnNextClicked)
 
 	/** Attribute controlling enabled state of back functionality */
 	SLATE_ATTRIBUTE(bool, IsBackEnabled)
@@ -59,6 +60,9 @@ class STutorialContent : public SCompoundWidget
 	/** Whether we can show full window content in this overlay (i.e. in the same window as the navigation controls) */
 	SLATE_ARGUMENT(bool, AllowNonWidgetContent)
 
+	/** Delegate for querying whether a widget was drawn */
+	SLATE_EVENT(FOnWasWidgetDrawn, OnWasWidgetDrawn)
+
 	SLATE_END_ARGS()
 
 	void Construct(const FArguments& InArgs, UEditorTutorial* InTutorial, const FTutorialContent& InContent);
@@ -209,4 +213,7 @@ private:
 
 	/** Whether we can show full window content in this overlay (i.e. in the same window as the navigation controls) */
 	bool bAllowNonWidgetContent;
+
+	/** Delegate for querying whether a widget was drawn */
+	FOnWasWidgetDrawn OnWasWidgetDrawn;
 };
\ No newline at end of file
diff --git a/Engine/Source/Editor/IntroTutorials/Private/STutorialOverlay.cpp b/Engine/Source/Editor/IntroTutorials/Private/STutorialOverlay.cpp
--- a/Engine/Source/Editor/IntroTutorials/Private/STutorialOverlay.cpp
+++ b/Engine/Source/Editor/IntroTutorials/Private/STutorialOverlay.cpp
@@ -18,6 +18,7 @@ void STutorialOverlay::Construct(const FArguments& InArgs, UEditorTutorial* InTu
 	bIsStandalone = InArgs._IsStandalone;
 	OnClosed = InArgs._OnClosed;
 	bHasValidContent = InStage != nullptr;
+	OnWidgetWasDrawn = InArgs._OnWidgetWasDrawn;
 
 	TSharedPtr<SOverlay> Overlay;
 
@@ -80,7 +81,8 @@ void STutorialOverlay::Construct(const FArguments& InArgs, UEditorTutorial* InTu
 						.IsNextEnabled(InArgs._IsNextEnabled)
 						.WrapTextAt(WidgetContent.ContentWidth)
 						.Anchor(WidgetContent.WidgetAnchor)
-						.AllowNonWidgetContent(InArgs._AllowNonWidgetContent);
+						.AllowNonWidgetContent(InArgs._AllowNonWidgetContent)
+						.OnWasWidgetDrawn(InArgs._OnWasWidgetDrawn);
 					
 					PerformWidgetInteractions(WidgetContent); 					
 
@@ -132,6 +134,7 @@ int32 STutorialOverlay::TraverseWidgets(TSharedRef<SWidget> InWidget, const FGeo
 	{
 		// we are a named widget - ask it to draw
 		OnPaintNamedWidget.Broadcast(InWidget, InGeometry);
+		OnWidgetWasDrawn.ExecuteIfBound(Tag);
 
 		// if we are picking, we need to draw an outline here
 		FName WidgetNameToHighlight = NAME_None;
diff --git a/Engine/Source/Editor/IntroTutorials/Private/STutorialOverlay.h b/Engine/Source/Editor/IntroTutorials/Private/STutorialOverlay.h
--- a/Engine/Source/Editor/IntroTutorials/Private/STutorialOverlay.h
+++ b/Engine/Source/Editor/IntroTutorials/Private/STutorialOverlay.h
@@ -4,14 +4,16 @@
 
 #include "EditorTutorial.h"
 
-class STutorialContent;
-
 /** Delegate used when drawing/arranging widgets */
 DECLARE_MULTICAST_DELEGATE_TwoParams(FOnPaintNamedWidget, TSharedRef<SWidget> /*InWidget*/, const FGeometry& /*InGeometry*/);
 
 /** Delegate used to inform widgets of the current window size, so they can auto-adjust layout */
 DECLARE_MULTICAST_DELEGATE_OneParam(FOnCacheWindowSize, const FVector2D& /*InWindowSize*/);
 
+/** Delegates for registering & querying whether a widget was drawn */
+DECLARE_DELEGATE_OneParam(FOnWidgetWasDrawn, const FName& /*InName*/);
+DECLARE_DELEGATE_RetVal_OneParam(bool, FOnWasWidgetDrawn, const FName& /*InName*/);
+
 /**
  * The widget which displays multiple 'floating' pieces of content overlaid onto the editor
  */
@@ -30,11 +32,11 @@ public:
 	SLATE_ARGUMENT(bool, IsStandalone)
 
 	/** Delegate fired when the close button is clicked */
-	SLATE_ARGUMENT(FSimpleDelegate, OnClosed)
+	SLATE_EVENT(FSimpleDelegate, OnClosed)
 
-	SLATE_ARGUMENT(FSimpleDelegate, OnBackClicked)
-	SLATE_ARGUMENT(FSimpleDelegate, OnHomeClicked)
-	SLATE_ARGUMENT(FSimpleDelegate, OnNextClicked)
+	SLATE_EVENT(FSimpleDelegate, OnBackClicked)
+	SLATE_EVENT(FSimpleDelegate, OnHomeClicked)
+	SLATE_EVENT(FSimpleDelegate, OnNextClicked)
 	SLATE_ATTRIBUTE(bool, IsBackEnabled)
 	SLATE_ATTRIBUTE(bool, IsHomeEnabled)
 	SLATE_ATTRIBUTE(bool, IsNextEnabled)
@@ -42,6 +44,10 @@ public:
 	/** Whether we can show full window content in this overlay (i.e. in the same window as the navigation controls) */
 	SLATE_ARGUMENT(bool, AllowNonWidgetContent)
 
+	/** Delegates for registering & querying whether a widget was drawn */
+	SLATE_EVENT(FOnWidgetWasDrawn, OnWidgetWasDrawn)
+	SLATE_EVENT(FOnWasWidgetDrawn, OnWasWidgetDrawn)
+
 	SLATE_END_ARGS()
 
 	void Construct(const FArguments& InArgs, UEditorTutorial* InTutorial, FTutorialStage* const InStage);
@@ -86,4 +92,7 @@ private:
 
 	/** Flag to see if we have valid content (this widget is created to also supply picker overlays) */
 	bool bHasValidContent;
+
+	/** Delegate for querying whether a widget was drawn */
+	FOnWidgetWasDrawn OnWidgetWasDrawn;
 };
\ No newline at end of file
diff --git a/Engine/Source/Editor/IntroTutorials/Private/STutorialRoot.cpp b/Engine/Source/Editor/IntroTutorials/Private/STutorialRoot.cpp
--- a/Engine/Source/Editor/IntroTutorials/Private/STutorialRoot.cpp
+++ b/Engine/Source/Editor/IntroTutorials/Private/STutorialRoot.cpp
@@ -52,6 +52,8 @@ void STutorialRoot::MaybeAddOverlay(TSharedRef<SWindow> InWindow)
 					.OnGetCurrentTutorial(FOnGetCurrentTutorial::CreateSP(this, &STutorialRoot::HandleGetCurrentTutorial))
 					.OnGetCurrentTutorialStage(FOnGetCurrentTutorialStage::CreateSP(this, &STutorialRoot::HandleGetCurrentTutorialStage))
 					.OnLaunchTutorial(FOnLaunchTutorial::CreateSP(this, &STutorialRoot::LaunchTutorial))
+					.OnWasWidgetDrawn(FOnWasWidgetDrawn::CreateSP(this, &STutorialRoot::WasWidgetDrawn))
+					.OnWidgetWasDrawn(FOnWidgetWasDrawn::CreateSP(this, &STutorialRoot::WidgetWasDrawn))
 				]
 			];
 
@@ -75,6 +77,11 @@ void STutorialRoot::Tick(const FGeometry& AllottedGeometry, const double InCurre
 	{
 		MaybeAddOverlay(Window);
 	}
+
+	// empty array but leave us the slack (we dont want to reallocate all the time, and this array should never grow too large)
+	PreviouslyDrawnWidgets.Empty(PreviouslyDrawnWidgets.Max());
+	PreviouslyDrawnWidgets.Append(MoveTemp(DrawnWidgets));
+	DrawnWidgets.Empty(DrawnWidgets.Max());
 }
 
 void STutorialRoot::SummonTutorialBrowser(TWeakPtr<SWindow> InWindow, const FString& InFilter)
@@ -316,4 +323,22 @@ void STutorialRoot::HandleCloseClicked()
 	}
 }
 
+bool STutorialRoot::WasWidgetDrawn(const FName& InName) const
+{
+	for(const auto& WidgetName : PreviouslyDrawnWidgets)
+	{
+		if(InName == WidgetName)
+		{
+			return true;
+		}
+	}
+
+	return false;
+}
+
+void STutorialRoot::WidgetWasDrawn(const FName& InName)
+{
+	DrawnWidgets.Add(InName);
+}
+
 #undef LOCTEXT_NAMESPACE
\ No newline at end of file
diff --git a/Engine/Source/Editor/IntroTutorials/Private/STutorialRoot.h b/Engine/Source/Editor/IntroTutorials/Private/STutorialRoot.h
--- a/Engine/Source/Editor/IntroTutorials/Private/STutorialRoot.h
+++ b/Engine/Source/Editor/IntroTutorials/Private/STutorialRoot.h
@@ -45,6 +45,12 @@ public:
 	/** Got to the next stage in the current tutorial */
 	void GoToNextStage(TWeakPtr<SWindow> InNavigationWindow);
 
+	/** Has this named widget been raw anywhere this frame? */
+	bool WasWidgetDrawn(const FName& InName) const;
+
+	/** Register that this widget was drawn this frame */
+	void WidgetWasDrawn(const FName& InName);
+
 private:
 	/** Handle when the next button is clicked */
 	void HandleNextClicked(TWeakPtr<SWindow> InNavigationWindow);
@@ -79,4 +85,8 @@ private:
 
 	/** Start time of the current tutorial, if any */
 	float CurrentTutorialStartTime;
+
+	/** Current drawn widgets for the last frame */
+	TArray<FName> DrawnWidgets;
+	TArray<FName> PreviouslyDrawnWidgets;
 };
\ No newline at end of file
