CRTICIAL: Regression: Opening many of the asset types give the warning "Failed to find object 'Class None'.

#ttp 347342

diff --git a/Engine/Source/Editor/IntroTutorials/Private/EditorTutorialSettings.cpp b/Engine/Source/Editor/IntroTutorials/Private/EditorTutorialSettings.cpp
--- a/Engine/Source/Editor/IntroTutorials/Private/EditorTutorialSettings.cpp
+++ b/Engine/Source/Editor/IntroTutorials/Private/EditorTutorialSettings.cpp
@@ -17,13 +17,13 @@ void UEditorTutorialSettings::FindTutorialInfoForContext(FName InContext, UEdito
 		{
 			OutBrowserFilter = Context.BrowserFilter;
 
-			TSubclassOf<UEditorTutorial> LaunchTutorialClass = LoadClass<UEditorTutorial>(nullptr, *Context.LaunchTutorial.AssetLongPathname, nullptr, LOAD_None, nullptr);
+			TSubclassOf<UEditorTutorial> LaunchTutorialClass = Context.LaunchTutorial.ResolveClass();
 			if (LaunchTutorialClass != nullptr)
 			{
 				OutLaunchTutorial = LaunchTutorialClass->GetDefaultObject<UEditorTutorial>();
 			}
 
-			TSubclassOf<UEditorTutorial> AttractTutorialClass = LoadClass<UEditorTutorial>(nullptr, *Context.AttractTutorial.AssetLongPathname, nullptr, LOAD_None, nullptr);
+			TSubclassOf<UEditorTutorial> AttractTutorialClass = Context.AttractTutorial.ResolveClass();
 			if (AttractTutorialClass != nullptr)
 			{
 				OutAttractTutorial = AttractTutorialClass->GetDefaultObject<UEditorTutorial>();
