Added breadcrumb trail to tutorials browser

Displays crumbs according to the current navigation filter.

TTP# 347533 - TUTORIALS: Display tree path in Tutorial Window

diff --git a/Engine/Source/Editor/EditorStyle/Private/SlateEditorStyle.cpp b/Engine/Source/Editor/EditorStyle/Private/SlateEditorStyle.cpp
--- a/Engine/Source/Editor/EditorStyle/Private/SlateEditorStyle.cpp
+++ b/Engine/Source/Editor/EditorStyle/Private/SlateEditorStyle.cpp
@@ -2512,6 +2512,10 @@ void FSlateEditorStyle::FStyle::SetupTutorialStyles()
 
 		Set( "Tutorials.Browser.Completed", new IMAGE_BRUSH( "Tutorials/TutorialCompleted", Icon32x32 ) );
 
+		Set( "Tutorials.Browser.Breadcrumb", new IMAGE_BRUSH( "Tutorials/Breadcrumb", Icon8x8, FLinearColor(0.0f, 0.0f, 0.0f, 1.0f) ) );
+		Set( "Tutorials.Browser.PathText", FTextBlockStyle(TutorialText)
+			.SetFontSize(9));
+
 		Set( "Tutorials.Navigation.Button", FButtonStyle()
 			.SetNormal( BOX_BRUSH( "Common/ButtonHoverHint", FMargin(4/16.0f), FLinearColor(0,0,0,0) ) )
 			.SetHovered( BOX_BRUSH( "Common/ButtonHoverHint", FMargin(4/16.0f), FLinearColor(0,0,0,0) ) )
diff --git a/Engine/Source/Editor/IntroTutorials/Private/STutorialsBrowser.cpp b/Engine/Source/Editor/IntroTutorials/Private/STutorialsBrowser.cpp
--- a/Engine/Source/Editor/IntroTutorials/Private/STutorialsBrowser.cpp
+++ b/Engine/Source/Editor/IntroTutorials/Private/STutorialsBrowser.cpp
@@ -156,6 +156,12 @@ public:
 		return Title;
 	}
 
+
+	FText GetTitleText() const override
+	{
+		return !Category.Title.IsEmpty() ? Category.Title : FText::FromString(CategoryName);
+	}
+
 	bool SortAgainst(TSharedRef<ITutorialListEntry> OtherEntry) const override
 	{
 		return GetTitleString() < OtherEntry->GetTitleString();
@@ -359,6 +365,11 @@ public:
 		return bPassesFilter && bPassesCategory;
 	}
 
+	FText GetTitleText() const override
+	{
+		return Tutorial->Title;
+	}
+
 	FString GetTitleString() const override
 	{
 		return Tutorial->Title.ToString();
@@ -499,7 +510,7 @@ void STutorialsBrowser::Construct(const FArguments& InArgs)
 					.OnTextChanged(this, &STutorialsBrowser::OnSearchTextChanged)
 				]
 				+SVerticalBox::Slot()
-				.AutoHeight()
+				.FillHeight(1.0f)
 				[
 					SNew(SBox)
 					.HeightOverride(600.0f)
@@ -516,11 +527,25 @@ void STutorialsBrowser::Construct(const FArguments& InArgs)
 						]
 					]
 				]
+				+SVerticalBox::Slot()
+				.AutoHeight()
+				.Padding(2.0f)
+				[
+					SAssignNew(BreadcrumbTrail, SBreadcrumbTrail<TSharedPtr<ITutorialListEntry>>)
+					.ButtonContentPadding(FMargin(1.0f, 1.0f))
+					.DelimiterImage(FEditorStyle::GetBrush("Tutorials.Browser.Breadcrumb"))
+					.TextStyle(FEditorStyle::Get(), "Tutorials.Browser.PathText")
+					.ShowLeadingDelimiter( true )
+					.InvertTextColorOnHover( false )
+					.OnCrumbClicked(this, &STutorialsBrowser::OnBreadcrumbClicked)
+				]
 			]
 		]
 	];
 
 	ReloadTutorials();
+
+	RebuildCrumbs();
 }
 
 void STutorialsBrowser::SetFilter(const FString& InFilter)
@@ -729,6 +754,8 @@ FReply STutorialsBrowser::OnBackButtonClicked()
 		}
 	}
 
+	RebuildCrumbs();
+
 	return FReply::Handled();
 }
 
@@ -764,6 +791,8 @@ void STutorialsBrowser::OnCategorySelected(const FString& InCategory)
 {
 	NavigationFilter = InCategory;
 	FilterTutorials();
+
+	RebuildCrumbs();
 }
 
 void STutorialsBrowser::FilterTutorials()
@@ -809,17 +838,20 @@ void STutorialsBrowser::FilterTutorials()
 
 TSharedPtr<FTutorialListEntry_Category> STutorialsBrowser::FindCategory_Recursive(TSharedPtr<FTutorialListEntry_Category> InCategory) const
 {
-	if(InCategory->Category.Identifier == NavigationFilter)
+	if(InCategory.IsValid())
 	{
-		return InCategory;
-	}
+		if(InCategory->Category.Identifier == NavigationFilter)
+		{
+			return InCategory;
+		}
 
-	for(const auto& Category : InCategory->SubCategories)
-	{
-		TSharedPtr<FTutorialListEntry_Category> TestCategory = FindCategory_Recursive(StaticCastSharedPtr<FTutorialListEntry_Category>(Category));
-		if(TestCategory.IsValid())
+		for(const auto& Category : InCategory->SubCategories)
 		{
-			return TestCategory;
+			TSharedPtr<FTutorialListEntry_Category> TestCategory = FindCategory_Recursive(StaticCastSharedPtr<FTutorialListEntry_Category>(Category));
+			if(TestCategory.IsValid())
+			{
+				return TestCategory;
+			}
 		}
 	}
 
@@ -837,4 +869,60 @@ FText STutorialsBrowser::GetSearchText() const
 	return SearchFilter;
 }
 
+void STutorialsBrowser::OnBreadcrumbClicked(const TSharedPtr<ITutorialListEntry>& InEntry)
+{
+	TSharedPtr<ITutorialListEntry> ClickedEntry = InEntry;
+
+	if(ClickedEntry.IsValid())
+	{
+		NavigationFilter = StaticCastSharedPtr<FTutorialListEntry_Category>(ClickedEntry)->Category.Identifier;
+	}
+	else
+	{
+		NavigationFilter.Empty();
+	}
+	
+	RebuildCrumbs();
+
+	FilterTutorials();
+}
+
+void STutorialsBrowser::RebuildCrumbs()
+{
+	BreadcrumbTrail->ClearCrumbs();
+
+	// rebuild crumbs to this point
+	TArray<TSharedPtr<FTutorialListEntry_Category>> Entries;
+	TSharedPtr<FTutorialListEntry_Category> CurrentCategory = FindCategory_Recursive(RootEntry);
+	if(CurrentCategory.IsValid())
+	{
+		TSharedPtr<FTutorialListEntry_Category> Category = StaticCastSharedPtr<FTutorialListEntry_Category>(CurrentCategory);
+		while(Category.IsValid())
+		{
+			Entries.Add(Category);
+			if(Category->ParentCategory.IsValid())
+			{
+				Category = StaticCastSharedPtr<FTutorialListEntry_Category>(Category->ParentCategory.Pin());
+			}
+			else
+			{
+				break;
+			}
+		}
+	}
+
+	for(int32 Index = Entries.Num() - 1; Index >= 0; Index--)
+	{
+		TSharedPtr<FTutorialListEntry_Category> Entry = Entries[Index];
+		if(RootEntry == Entry)
+		{
+			BreadcrumbTrail->PushCrumb(LOCTEXT("PathRoot", "Tutorials"), TSharedPtr<ITutorialListEntry>());
+		}
+		else
+		{
+			BreadcrumbTrail->PushCrumb(Entry->GetTitleText(), Entry);
+		}
+	}
+}
+
 #undef LOCTEXT_NAMESPACE
\ No newline at end of file
diff --git a/Engine/Source/Editor/IntroTutorials/Private/STutorialsBrowser.h b/Engine/Source/Editor/IntroTutorials/Private/STutorialsBrowser.h
--- a/Engine/Source/Editor/IntroTutorials/Private/STutorialsBrowser.h
+++ b/Engine/Source/Editor/IntroTutorials/Private/STutorialsBrowser.h
@@ -15,6 +15,9 @@ struct ITutorialListEntry
 	/** Whether this entry passes the current filter criteria */
 	virtual bool PassesFilter(const FString& InCategoryFilter, const FString& InFilter) const = 0;
 
+	/** Get the text representation of this item's title */
+	virtual FText GetTitleText() const = 0;
+
 	/** Get the string representation of this item's title */
 	virtual FString GetTitleString() const = 0;
 
@@ -88,6 +91,12 @@ protected:
 	/** Supplies the text to display in the filter box */
 	FText GetSearchText() const;
 
+	/** Handle clicking the breadcrumb trail */
+	void OnBreadcrumbClicked(const TSharedPtr<ITutorialListEntry>& InEntry);
+
+	/** Rebuild the breadcrumb trail according to the current category */
+	void RebuildCrumbs();
+
 private:
 
 	/** Root entry of the tutorials tree */
@@ -119,4 +128,7 @@ private:
 
 	/** Parent window of this browser */
 	TWeakPtr<SWindow> ParentWindow;
+
+	/** Breadcrumb trail for path display */
+	TSharedPtr<SBreadcrumbTrail<TSharedPtr<ITutorialListEntry>>> BreadcrumbTrail;
 };
\ No newline at end of file
