Searching tutorials now recursively searches all categories and sub-categories

If we have a search string we should search all categories starting at the current context, rather than just one level deep as before.
Also fixed search text highlighting to be consistently coloured and to show up for rich text content.

TTP# 347534 - CRITICAL: TUTORIALS: Search bar only shows results for currently displayed tutorials

diff --git a/Engine/Source/Editor/EditorStyle/Private/SlateEditorStyle.cpp b/Engine/Source/Editor/EditorStyle/Private/SlateEditorStyle.cpp
--- a/Engine/Source/Editor/EditorStyle/Private/SlateEditorStyle.cpp
+++ b/Engine/Source/Editor/EditorStyle/Private/SlateEditorStyle.cpp
@@ -2467,10 +2467,12 @@ void FSlateEditorStyle::FStyle::SetupTutorialStyles()
 		const FLinearColor TutorialSelectionColor = FLinearColor(0.19f, 0.33f, 0.72f);
 
 		const FTextBlockStyle TutorialText = FTextBlockStyle(DocumentationText)
-			.SetColorAndOpacity(FLinearColor::Black);
+			.SetColorAndOpacity(FLinearColor::Black)
+			.SetHighlightColor(TutorialSelectionColor);
 
 		const FTextBlockStyle TutorialHeaderText = FTextBlockStyle(DocumentationHeaderText)
-			.SetColorAndOpacity(FLinearColor::Black);
+			.SetColorAndOpacity(FLinearColor::Black)
+			.SetHighlightColor(TutorialSelectionColor);
 
 		Set( "Tutorials.Border", new BOX_BRUSH( "Tutorials/OverlayFrame", FMargin(18.0f/64.0f), FLinearColor(1.0f, 1.0f, 1.0f, 1.0f) ) );
 		Set( "Tutorials.Browser.WelcomeHeader", FTextBlockStyle(TutorialHeaderText)
diff --git a/Engine/Source/Editor/IntroTutorials/Private/STutorialContent.cpp b/Engine/Source/Editor/IntroTutorials/Private/STutorialContent.cpp
--- a/Engine/Source/Editor/IntroTutorials/Private/STutorialContent.cpp
+++ b/Engine/Source/Editor/IntroTutorials/Private/STutorialContent.cpp
@@ -261,7 +261,7 @@ static TSharedRef<SWidget> GetStageTitle(const FExcerpt& InExcerpt, int32 InCurr
 	return SNullWidget::NullWidget;
 }
 
-TSharedRef<SWidget> STutorialContent::GenerateContentWidget(const FTutorialContent& InContent, float WrapTextAt, TSharedPtr<IDocumentationPage>& OutDocumentationPage)
+TSharedRef<SWidget> STutorialContent::GenerateContentWidget(const FTutorialContent& InContent, float WrapTextAt, TSharedPtr<IDocumentationPage>& OutDocumentationPage, const TAttribute<FText>& InHighlightText)
 {
 	// Style for the documentation
 	static FDocumentationStyle DocumentationStyle;
@@ -285,7 +285,9 @@ TSharedRef<SWidget> STutorialContent::GenerateContentWidget(const FTutorialConte
 				.Visibility(EVisibility::SelfHitTestInvisible)
 				.WrapTextAt(WrapTextAt)
 				.Text(InContent.Text)
-				.TextStyle(FEditorStyle::Get(), "Tutorials.Content");
+				.TextStyle(FEditorStyle::Get(), "Tutorials.Content")
+				.HighlightText(InHighlightText)
+				.HighlightColor(FEditorStyle::Get().GetColor("Tutorials.Browser.HighlightTextColor"));
 		}
 
 	case ETutorialContent::UDNExcerpt:
@@ -326,7 +328,8 @@ TSharedRef<SWidget> STutorialContent::GenerateContentWidget(const FTutorialConte
 					.Text(InContent.Text)
 					.WrapTextAt(WrapTextAt)
 					.Margin(4)
-					.LineHeightPercentage(1.1f);
+					.LineHeightPercentage(1.1f)
+					.HighlightText(InHighlightText);
 		}
 		break;
 	}
diff --git a/Engine/Source/Editor/IntroTutorials/Private/STutorialContent.h b/Engine/Source/Editor/IntroTutorials/Private/STutorialContent.h
--- a/Engine/Source/Editor/IntroTutorials/Private/STutorialContent.h
+++ b/Engine/Source/Editor/IntroTutorials/Private/STutorialContent.h
@@ -72,7 +72,7 @@ class STutorialContent : public SCompoundWidget
 	virtual FReply OnMouseButtonUp( const FGeometry& MyGeometry, const FPointerEvent& MouseEvent ) override;
 
 	/** Helper function to generate widgets from an FTutorialContent struct */
-	static TSharedRef<SWidget> GenerateContentWidget(const FTutorialContent& InContent, float WrapTextAt, TSharedPtr<IDocumentationPage>& OutDocumentationPage);
+	static TSharedRef<SWidget> GenerateContentWidget(const FTutorialContent& InContent, float WrapTextAt, TSharedPtr<IDocumentationPage>& OutDocumentationPage, const TAttribute<FText>& InHighlightText = TAttribute<FText>());
 
 	/** Handle repositioning the widget */
 	FVector2D GetPosition() const;
diff --git a/Engine/Source/Editor/IntroTutorials/Private/STutorialsBrowser.cpp b/Engine/Source/Editor/IntroTutorials/Private/STutorialsBrowser.cpp
--- a/Engine/Source/Editor/IntroTutorials/Private/STutorialsBrowser.cpp
+++ b/Engine/Source/Editor/IntroTutorials/Private/STutorialsBrowser.cpp
@@ -349,7 +349,7 @@ public:
 						+SVerticalBox::Slot()
 						.FillHeight(1.0f)
 						[
-							STutorialContent::GenerateContentWidget(Tutorial->SummaryContent, 500.0f, DocumentationPage)
+							STutorialContent::GenerateContentWidget(Tutorial->SummaryContent, 500.0f, DocumentationPage, HighlightText)
 						]
 					]
 				]
@@ -799,27 +799,66 @@ void STutorialsBrowser::FilterTutorials()
 {
 	FilteredEntries.Empty();
 
-	TSharedPtr<FTutorialListEntry_Category> CurrentCategory = FindCategory_Recursive(RootEntry);
-
-	if(CurrentCategory.IsValid())
+	if(SearchFilter.IsEmpty())
 	{
-		for(const auto& SubCategory : CurrentCategory->SubCategories)
+		TSharedPtr<FTutorialListEntry_Category> CurrentCategory = FindCategory_Recursive(RootEntry);
+
+		if(CurrentCategory.IsValid())
 		{
-			if(SubCategory->PassesFilter(CategoryFilter, SearchFilter.ToString()))
+			for(const auto& SubCategory : CurrentCategory->SubCategories)
 			{
-				FilteredEntries.Add(SubCategory);
+				if(SubCategory->PassesFilter(CategoryFilter, SearchFilter.ToString()))
+				{
+					FilteredEntries.Add(SubCategory);
+				}
 			}
-		}
 
-		for(const auto& Tutorial : CurrentCategory->Tutorials)
-		{
-			if(Tutorial->PassesFilter(CategoryFilter, SearchFilter.ToString()))
+			for(const auto& Tutorial : CurrentCategory->Tutorials)
 			{
-				FilteredEntries.Add(Tutorial);
+				if(Tutorial->PassesFilter(CategoryFilter, SearchFilter.ToString()))
+				{
+					FilteredEntries.Add(Tutorial);
+				}
 			}
+
+			CurrentCategoryPtr = CurrentCategory;
 		}
+	}
+	else
+	{
+		struct Local
+		{
+			static void AddSubCategory_Recursive(const FString& InCategoryFilter, const FString& InSearchFilter, TSharedPtr<FTutorialListEntry_Category> InCategory, TArray<TSharedPtr<ITutorialListEntry>>& InOutFilteredEntries)
+			{
+				if(InCategory.IsValid())
+				{
+					for(const auto& SubCategory : InCategory->SubCategories)
+					{
+						if(SubCategory->PassesFilter(InCategoryFilter, InSearchFilter))
+						{
+							InOutFilteredEntries.Add(SubCategory);
+						}
+
+						AddSubCategory_Recursive(InCategoryFilter, InSearchFilter, StaticCastSharedPtr<FTutorialListEntry_Category>(SubCategory), InOutFilteredEntries);
+					}
+
+					for(const auto& Tutorial : InCategory->Tutorials)
+					{
+						if(Tutorial->PassesFilter(InCategoryFilter, InSearchFilter))
+						{
+							InOutFilteredEntries.Add(Tutorial);
+						}
+					}
+				}
+			};
+		};
 
-		CurrentCategoryPtr = CurrentCategory;
+		TSharedPtr<FTutorialListEntry_Category> CurrentCategory = FindCategory_Recursive(RootEntry);
+		if(CurrentCategory.IsValid())
+		{
+			Local::AddSubCategory_Recursive(CategoryFilter, SearchFilter.ToString(), CurrentCategory, FilteredEntries);
+			CurrentCategoryPtr = CurrentCategory;
+		}
 	}
 
 	FilteredEntries.Sort(
