Added support for Back and Forward mouse buttons to Tutorials

As with the content browser, these buttons only function as back and forward when they are clicked within the bounds of the widget.

TTP# 347536 - TUTORIALS: Add Forward/Back web browser controls

diff --git a/Engine/Source/Editor/IntroTutorials/Private/STutorialContent.cpp b/Engine/Source/Editor/IntroTutorials/Private/STutorialContent.cpp
--- a/Engine/Source/Editor/IntroTutorials/Private/STutorialContent.cpp
+++ b/Engine/Source/Editor/IntroTutorials/Private/STutorialContent.cpp
@@ -651,4 +651,50 @@ FSlateColor STutorialContent::GetNextButtonColor() const
 	return NextButton->IsHovered() ? FLinearColor(0.1f, 0.1f, 0.1f, 1.0f) : FLinearColor(0.0f, 0.0f, 0.0f, 1.0f);
 }
 
+FReply STutorialContent::OnMouseButtonDown( const FGeometry& MyGeometry, const FPointerEvent& MouseEvent )
+{
+	// Mouse back and forward buttons traverse history
+	if ( MouseEvent.GetEffectingButton() == EKeys::ThumbMouseButton)
+	{
+		if(IsBackEnabled.Get())
+		{
+			OnBackClicked.ExecuteIfBound();
+			return FReply::Handled();
+		}
+	}
+	else if ( MouseEvent.GetEffectingButton() == EKeys::ThumbMouseButton2)
+	{
+		if(IsNextEnabled.Get())
+		{
+			OnNextClicked.ExecuteIfBound();
+			return FReply::Handled();
+		}
+	}
+
+	return FReply::Unhandled();
+}
+
+FReply STutorialContent::OnMouseButtonDoubleClick( const FGeometry& InMyGeometry, const FPointerEvent& InMouseEvent )
+{
+	// Mouse back and forward buttons traverse history
+	if ( InMouseEvent.GetEffectingButton() == EKeys::ThumbMouseButton)
+	{
+		if(IsBackEnabled.Get())
+		{
+			OnBackClicked.ExecuteIfBound();
+			return FReply::Handled();
+		}
+	}
+	else if ( InMouseEvent.GetEffectingButton() == EKeys::ThumbMouseButton2)
+	{
+		if(IsNextEnabled.Get())
+		{
+			OnNextClicked.ExecuteIfBound();
+			return FReply::Handled();
+		}
+	}
+
+	return FReply::Unhandled();
+}
+
 #undef LOCTEXT_NAMESPACE
\ No newline at end of file
diff --git a/Engine/Source/Editor/IntroTutorials/Private/STutorialContent.h b/Engine/Source/Editor/IntroTutorials/Private/STutorialContent.h
--- a/Engine/Source/Editor/IntroTutorials/Private/STutorialContent.h
+++ b/Engine/Source/Editor/IntroTutorials/Private/STutorialContent.h
@@ -70,6 +70,8 @@ class STutorialContent : public SCompoundWidget
 	/** SWidget implementation */
 	virtual int32 OnPaint( const FPaintArgs& Args, const FGeometry& AllottedGeometry, const FSlateRect& MyClippingRect, FSlateWindowElementList& OutDrawElements, int32 LayerId, const FWidgetStyle& InWidgetStyle, bool bParentEnabled ) const override;
 	virtual FReply OnMouseButtonUp( const FGeometry& MyGeometry, const FPointerEvent& MouseEvent ) override;
+	virtual FReply OnMouseButtonDown( const FGeometry& MyGeometry, const FPointerEvent& MouseEvent ) override;
+	virtual FReply OnMouseButtonDoubleClick( const FGeometry& InMyGeometry, const FPointerEvent& InMouseEvent ) override;
 
 	/** Helper function to generate widgets from an FTutorialContent struct */
 	static TSharedRef<SWidget> GenerateContentWidget(const FTutorialContent& InContent, float WrapTextAt, TSharedPtr<IDocumentationPage>& OutDocumentationPage, const TAttribute<FText>& InHighlightText = TAttribute<FText>());
@@ -139,7 +141,7 @@ private:
 	/** Delegate handler allowing us to change the tootlip of the 'next' button depending on context */
 	FText GetNextButtonTooltip() const;
 
-	/** Chage next button color based on hover state */
+	/** Change next button color based on hover state */
 	FSlateColor GetNextButtonColor() const;
 
 private:
