Fixed auto-wrapped text in the tutorials list never growing again if it was shrunk

TTP# 347884 - EDITOR: TUTORIALS: text does not expand when maximizing the window

The auto-wrapping allowed the text block to shrink when the window was shrunk, however the sizers were set to align their content left horizontally (rather than fill, which is the default) meaning that the text block was never allowed to grow again once it had shrunk.

A better way to achieve this is to set the sizer slots to fill, and then use the text justification on the text block to control the text location - thankfully both of these are the defaults.

diff --git a/Engine/Source/Editor/IntroTutorials/Private/STutorialsBrowser.cpp b/Engine/Source/Editor/IntroTutorials/Private/STutorialsBrowser.cpp
--- a/Engine/Source/Editor/IntroTutorials/Private/STutorialsBrowser.cpp
+++ b/Engine/Source/Editor/IntroTutorials/Private/STutorialsBrowser.cpp
@@ -101,12 +101,10 @@ public:
 				+SHorizontalBox::Slot()
 				.FillWidth(1.0f)
 				.MaxWidth(600.0f)
-				.HAlign(HAlign_Left)
 				.VAlign(VAlign_Center)
 				[
 					SNew(SVerticalBox)
 					+SVerticalBox::Slot()
-					.HAlign(HAlign_Left)
 					.AutoHeight()
 					[
 						SNew(STextBlock)
@@ -117,7 +115,6 @@ public:
 						.HighlightShape(FEditorStyle::Get().GetBrush("TextBlock.HighlightShape"))
 					]
 					+SVerticalBox::Slot()
-					.HAlign(HAlign_Left)
 					.FillHeight(1.0f)
 					[
 						SNew(STextBlock)
