Tutorial content 'Next' button is now green, has text and is customizable

Rearranged layout/padding a little to compensate for larger apparent button size.
Kept larger button hitbox, which required some custom border handling for the button.

TTP# 348109 - EDITOR: TUTORIALS: Add green buttons to tutorials

diff --git a/Engine/Source/Editor/EditorStyle/Private/SlateEditorStyle.cpp b/Engine/Source/Editor/EditorStyle/Private/SlateEditorStyle.cpp
--- a/Engine/Source/Editor/EditorStyle/Private/SlateEditorStyle.cpp
+++ b/Engine/Source/Editor/EditorStyle/Private/SlateEditorStyle.cpp
@@ -2465,6 +2465,8 @@ void FSlateEditorStyle::FStyle::SetupTutorialStyles()
 	{
 		const FLinearColor TutorialButtonColor = FLinearColor(0.15f, 0.15f, 0.15f, 1.0f);
 		const FLinearColor TutorialSelectionColor = FLinearColor(0.19f, 0.33f, 0.72f);
+		const FLinearColor TutorialNavigationButtonColor = FLinearColor(0.0f, 0.59f, 0.14f, 1.0f);
+		const FLinearColor TutorialNavigationButtonHoverColor = FLinearColor(0.2f, 0.79f, 0.34f, 1.0f);
 
 		const FTextBlockStyle TutorialText = FTextBlockStyle(DocumentationText)
 			.SetColorAndOpacity(FLinearColor::Black)
@@ -2502,6 +2504,24 @@ void FSlateEditorStyle::FStyle::SetupTutorialStyles()
 			.SetNormalPadding( FMargin(0,0,0,1))
 			.SetPressedPadding( FMargin(0,1,0,0)));
 
+		Set( "Tutorials.Content.NavigationButtonWrapper", FButtonStyle()
+			.SetNormal( FSlateNoResource() )
+			.SetHovered( FSlateNoResource() )
+			.SetPressed( FSlateNoResource() )
+			.SetDisabled( FSlateNoResource() )
+			.SetNormalPadding( FMargin(0,0,0,1))
+			.SetPressedPadding( FMargin(0,1,0,0)));
+
+		Set( "Tutorials.Content.NavigationButton", FButtonStyle()
+			.SetNormal( BOX_BRUSH( "Common/ButtonHoverHint", FMargin(4/16.0f), TutorialNavigationButtonColor ) )
+			.SetHovered( BOX_BRUSH( "Common/ButtonHoverHint", FMargin(4/16.0f), TutorialNavigationButtonHoverColor ) )
+			.SetPressed( BOX_BRUSH( "Common/ButtonHoverHint", FMargin(4/16.0f), TutorialNavigationButtonHoverColor ) )
+			.SetDisabled( BOX_BRUSH( "Common/ButtonHoverHint", FMargin(4/16.0f), TutorialNavigationButtonColor ) )
+			.SetNormalPadding( FMargin(0,0,0,1))
+			.SetPressedPadding( FMargin(0,1,0,0)));
+
+		Set( "Tutorials.Content.NavigationText", FTextBlockStyle(TutorialText));
+
 		Set( "Tutorials.Content.Color", FLinearColor(1.0f,1.0f,1.0f,0.9f) );
 		Set( "Tutorials.Content.Color.Hovered", FLinearColor(1.0f,1.0f,1.0f,1.0f) );
 
diff --git a/Engine/Source/Editor/IntroTutorials/Private/STutorialContent.cpp b/Engine/Source/Editor/IntroTutorials/Private/STutorialContent.cpp
--- a/Engine/Source/Editor/IntroTutorials/Private/STutorialContent.cpp
+++ b/Engine/Source/Editor/IntroTutorials/Private/STutorialContent.cpp
@@ -19,7 +19,7 @@ namespace TutorialConstants
 	const float ShadowScale = 8.0f;
 	const float MaxBorderOffset = 8.0f;
 	const FMargin BorderSizeStandalone(24.0f, 24.0f);
-	const FMargin BorderSize(24.0f, 24.0f, 24.0f, 42.0f);
+	const FMargin BorderSize(24.0f, 24.0f, 24.0f, 52.0f);
 }
 
 const float ContentOffset = 10.0f;
@@ -29,7 +29,7 @@ void STutorialContent::Construct(const FArguments& InArgs, UEditorTutorial* InTu
 	bIsVisible = Anchor.Type == ETutorialAnchorIdentifier::None;
 
 	Tutorial = InTutorial;
-
+	
 	VerticalAlignment = InArgs._VAlign;
 	HorizontalAlignment = InArgs._HAlign;
 	WidgetOffset = InArgs._Offset;
@@ -44,6 +44,7 @@ void STutorialContent::Construct(const FArguments& InArgs, UEditorTutorial* InTu
 	Anchor = InArgs._Anchor;
 	bAllowNonWidgetContent = InArgs._AllowNonWidgetContent;
 	OnWasWidgetDrawn = InArgs._OnWasWidgetDrawn;
+	NextButtonText = InArgs._NextButtonText;
 
 	BorderIntroAnimation.AddCurve(0.0f, TutorialConstants::BorderIntroAnimationLength, ECurveEaseFunction::CubicOut);
 	BorderPulseAnimation.AddCurve(0.0f, TutorialConstants::BorderPulseAnimationLength, ECurveEaseFunction::Linear);
@@ -142,21 +143,40 @@ void STutorialContent::Construct(const FArguments& InArgs, UEditorTutorial* InTu
 			+SOverlay::Slot()
 			.VAlign(VAlign_Bottom)
 			.HAlign(HAlign_Right)
-			.Padding(8.0f)
+			.Padding(12.0f)
 			[
 				SAssignNew(NextButton, SButton)
 				.ToolTipText(this, &STutorialContent::GetNextButtonTooltip)
 				.OnClicked(this, &STutorialContent::HandleNextClicked)
 				.Visibility(this, &STutorialContent::GetMenuButtonVisibility)
-				.ButtonStyle(&FEditorStyle::Get().GetWidgetStyle<FButtonStyle>("Tutorials.Navigation.Button"))
+				.ButtonStyle(&FEditorStyle::Get().GetWidgetStyle<FButtonStyle>("Tutorials.Content.NavigationButtonWrapper"))
 				.ContentPadding(0.0f)
 				[
 					SNew(SBox)
 					.Padding(8.0f)
 					[
-						SNew(SImage)
-						.Image(this, &STutorialContent::GetNextButtonBrush)
-						.ColorAndOpacity(this, &STutorialContent::GetNextButtonColor)
+						SNew(SBorder)
+						.BorderImage(this, &STutorialContent::GetNextButtonBorder)
+						[
+							SNew(SHorizontalBox)
+							+SHorizontalBox::Slot()
+							.AutoWidth()
+							.VAlign(VAlign_Center)
+							.Padding(4.0f, 0.0f, 0.0f, 0.0f)
+							[
+								SNew(STextBlock)
+								.Text(this, &STutorialContent::GetNextButtonLabel)
+								.TextStyle(FEditorStyle::Get(), "Tutorials.Content.NavigationText")
+								.ColorAndOpacity(FLinearColor::White)
+							]
+							+SHorizontalBox::Slot()
+							.AutoWidth()
+							[
+								SNew(SImage)
+								.Image(this, &STutorialContent::GetNextButtonBrush)
+								.ColorAndOpacity(FLinearColor::White)
+							]
+						]
 					]
 				]
 			]
@@ -646,9 +666,28 @@ FText STutorialContent::GetNextButtonTooltip() const
 	}
 }
 
-FSlateColor STutorialContent::GetNextButtonColor() const
+FText STutorialContent::GetNextButtonLabel() const
+{
+	if(!NextButtonText.Get().IsEmpty())
+	{
+		return NextButtonText.Get();
+	}
+	else
+	{
+		if(IsNextEnabled.Get())
+		{
+			return LOCTEXT("DefaultNextButtonLabel", "Next");
+		}
+		else
+		{
+			return LOCTEXT("DefaultHomeButtonLabel", "Home");
+		}
+	}
+}
+
+const FSlateBrush* STutorialContent::GetNextButtonBorder() const
 {
-	return NextButton->IsHovered() ? FLinearColor(0.1f, 0.1f, 0.1f, 1.0f) : FLinearColor(0.0f, 0.0f, 0.0f, 1.0f);
+	return NextButton->IsHovered() ? &FEditorStyle::Get().GetWidgetStyle<FButtonStyle>("Tutorials.Content.NavigationButton").Hovered : &FEditorStyle::Get().GetWidgetStyle<FButtonStyle>("Tutorials.Content.NavigationButton").Normal;
 }
 
 FReply STutorialContent::OnMouseButtonDown( const FGeometry& MyGeometry, const FPointerEvent& MouseEvent )
diff --git a/Engine/Source/Editor/IntroTutorials/Private/STutorialContent.h b/Engine/Source/Editor/IntroTutorials/Private/STutorialContent.h
--- a/Engine/Source/Editor/IntroTutorials/Private/STutorialContent.h
+++ b/Engine/Source/Editor/IntroTutorials/Private/STutorialContent.h
@@ -63,6 +63,9 @@ class STutorialContent : public SCompoundWidget
 	/** Delegate for querying whether a widget was drawn */
 	SLATE_EVENT(FOnWasWidgetDrawn, OnWasWidgetDrawn)
 
+	/** Text to display on next/home button */
+	SLATE_ATTRIBUTE(FText, NextButtonText)
+
 	SLATE_END_ARGS()
 
 	void Construct(const FArguments& InArgs, UEditorTutorial* InTutorial, const FTutorialContent& InContent);
@@ -142,7 +145,10 @@ private:
 	FText GetNextButtonTooltip() const;
 
 	/** Change next button color based on hover state */
-	FSlateColor GetNextButtonColor() const;
+	FText GetNextButtonLabel() const;
+
+	/** We need to override the border ourselves, rather than let the button handle it, as we are using a larger apparent hitbox */
+	const FSlateBrush* GetNextButtonBorder() const;
 
 private:
 
@@ -218,4 +224,7 @@ private:
 
 	/** Delegate for querying whether a widget was drawn */
 	FOnWasWidgetDrawn OnWasWidgetDrawn;
+
+	/** Text for next/home button */
+	TAttribute<FText> NextButtonText;
 };
\ No newline at end of file
diff --git a/Engine/Source/Editor/IntroTutorials/Private/STutorialOverlay.cpp b/Engine/Source/Editor/IntroTutorials/Private/STutorialOverlay.cpp
--- a/Engine/Source/Editor/IntroTutorials/Private/STutorialOverlay.cpp
+++ b/Engine/Source/Editor/IntroTutorials/Private/STutorialOverlay.cpp
@@ -82,7 +82,8 @@ void STutorialOverlay::Construct(const FArguments& InArgs, UEditorTutorial* InTu
 						.WrapTextAt(WidgetContent.ContentWidth)
 						.Anchor(WidgetContent.WidgetAnchor)
 						.AllowNonWidgetContent(InArgs._AllowNonWidgetContent)
-						.OnWasWidgetDrawn(InArgs._OnWasWidgetDrawn);
+						.OnWasWidgetDrawn(InArgs._OnWasWidgetDrawn)
+						.NextButtonText(InStage->NextButtonText);
 					
 					PerformWidgetInteractions(WidgetContent); 					
 
diff --git a/Engine/Source/Editor/IntroTutorials/Public/EditorTutorial.h b/Engine/Source/Editor/IntroTutorials/Public/EditorTutorial.h
--- a/Engine/Source/Editor/IntroTutorials/Public/EditorTutorial.h
+++ b/Engine/Source/Editor/IntroTutorials/Public/EditorTutorial.h
@@ -188,14 +188,21 @@ struct INTROTUTORIALS_API FTutorialStage
 {
 	GENERATED_USTRUCT_BODY()
 
+	/** Identifier for this stage */
 	UPROPERTY(EditAnywhere, Category="Stage")
 	FName Name;
 
+	/** Non-widget-bound content to display in this stage */
 	UPROPERTY(EditAnywhere, Category="Stage")
 	FTutorialContent Content;
 
+	/** Widget-bound content to display for this stage */
 	UPROPERTY(EditAnywhere, Category="Stage")
 	TArray<FTutorialWidgetContent> WidgetContent;
+
+	/** Text to display on the next button */
+	UPROPERTY(EditAnywhere, Category="Stage")
+	FText NextButtonText;
 };
 
 /** An asset used to build a stage-by-stage tutorial in the editor */
