Tutorials are now restarted when taken if they were previously completed

Override parameter passed in to LaunchTutorial when the tutorial has been completed already.

TTP# 348110 - EDITOR: TUTORIALS: Completed tutorials should always restart when re-taken

diff --git a/Engine/Source/Editor/IntroTutorials/Private/STutorialRoot.cpp b/Engine/Source/Editor/IntroTutorials/Private/STutorialRoot.cpp
--- a/Engine/Source/Editor/IntroTutorials/Private/STutorialRoot.cpp
+++ b/Engine/Source/Editor/IntroTutorials/Private/STutorialRoot.cpp
@@ -106,6 +106,12 @@ void STutorialRoot::LaunchTutorial(UEditorTutorial* InTutorial, bool bInRestart,
 	{
 		CurrentTutorial = InTutorial;
 
+		// we force a restart if this tutorial was completed
+		if(GetDefault<UTutorialStateSettings>()->HaveCompletedTutorial(CurrentTutorial))
+		{
+			bInRestart = true;
+		}
+
 		bool bHaveSeenTutorial = false;
 		CurrentTutorialStage = bInRestart ? 0 : GetDefault<UTutorialStateSettings>()->GetProgress(CurrentTutorial, bHaveSeenTutorial);
 
