Fixed tutorial editor's widget picker on Mac

diff --git a/Engine/Source/Editor/IntroTutorials/Private/TutorialStructCustomization.cpp b/Engine/Source/Editor/IntroTutorials/Private/TutorialStructCustomization.cpp
--- a/Engine/Source/Editor/IntroTutorials/Private/TutorialStructCustomization.cpp
+++ b/Engine/Source/Editor/IntroTutorials/Private/TutorialStructCustomization.cpp
@@ -388,8 +388,14 @@ private:
 				);
 
 			TSharedPtr<SWindow> RootWindow = FGlobalTabmanager::Get()->GetRootWindow();
-			check(RootWindow.IsValid());
-			FSlateApplication::Get().AddWindowAsNativeChild(NewWindow, RootWindow.ToSharedRef());
+			if (RootWindow.IsValid())
+			{
+				FSlateApplication::Get().AddWindowAsNativeChild(NewWindow, RootWindow.ToSharedRef());
+			}
+			else
+			{
+				FSlateApplication::Get().AddWindow(NewWindow);
+			}
 
 			FIntroTutorials& IntroTutorials = FModuleManager::Get().GetModuleChecked<FIntroTutorials>("IntroTutorials");
 			IntroTutorials.OnIsPicking().BindSP(this, &SWidgetPicker::OnIsPicking);
