Adding external link indicator to tutorial hyperlinks

Slate changes:
Allowed hyperlink decorator and hyperlink run to be more easily derived from by moving some functions/variables to protected.

Added new run & decorator for hyperlinks in tutorials that add an icon to the end of the run if required.
Only use icon-style run if in display mode (not in edit mode).

diff --git a/Engine/Source/Editor/IntroTutorials/Private/STutorialContent.cpp b/Engine/Source/Editor/IntroTutorials/Private/STutorialContent.cpp
--- a/Engine/Source/Editor/IntroTutorials/Private/STutorialContent.cpp
+++ b/Engine/Source/Editor/IntroTutorials/Private/STutorialContent.cpp
@@ -338,7 +338,8 @@ TSharedRef<SWidget> STutorialContent::GenerateContentWidget(const FTutorialConte
 	case ETutorialContent::RichText:
 		{
 			TArray< TSharedRef< class ITextDecorator > > Decorators;
-			FTutorialText::GetRichTextDecorators(Decorators);
+			const bool bForEditing = false;
+			FTutorialText::GetRichTextDecorators(bForEditing, Decorators);
 
 			return SNew(SRichTextBlock)
 					.Visibility(EVisibility::SelfHitTestInvisible)
diff --git a/Engine/Source/Editor/IntroTutorials/Private/STutorialEditableText.cpp b/Engine/Source/Editor/IntroTutorials/Private/STutorialEditableText.cpp
--- a/Engine/Source/Editor/IntroTutorials/Private/STutorialEditableText.cpp
+++ b/Engine/Source/Editor/IntroTutorials/Private/STutorialEditableText.cpp
@@ -34,7 +34,8 @@ void STutorialEditableText::Construct(const FArguments& InArgs)
 
 
 	TArray<TSharedRef<class ITextDecorator>> Decorators;
-	FTutorialText::GetRichTextDecorators(Decorators);
+	const bool bForEditing = true;
+	FTutorialText::GetRichTextDecorators(bForEditing, Decorators);
 
 	for(auto& Decorator : Decorators)
 	{
diff --git a/Engine/Source/Editor/IntroTutorials/Private/TutorialHyperlinkDecorator.cpp b/Engine/Source/Editor/IntroTutorials/Private/TutorialHyperlinkDecorator.cpp
new file mode 100644
--- /dev/null
+++ b/Engine/Source/Editor/IntroTutorials/Private/TutorialHyperlinkDecorator.cpp
@@ -0,0 +1,45 @@
+// Copyright 1998-2014 Epic Games, Inc. All Rights Reserved.
+
+#include "IntroTutorialsPrivatePCH.h"
+#include "TutorialHyperlinkDecorator.h"
+#include "TutorialHyperlinkRun.h"
+
+TSharedRef< FTutorialHyperlinkDecorator > FTutorialHyperlinkDecorator::Create( FString Id, const FSlateHyperlinkRun::FOnClick& NavigateDelegate, const FSlateHyperlinkRun::FOnGetTooltipText& InToolTipTextDelegate, const FSlateHyperlinkRun::FOnGenerateTooltip& InToolTipDelegate )
+{
+	return MakeShareable( new FTutorialHyperlinkDecorator( Id, NavigateDelegate, InToolTipTextDelegate, InToolTipDelegate ) );
+}
+
+FTutorialHyperlinkDecorator::FTutorialHyperlinkDecorator( FString InId, const FSlateHyperlinkRun::FOnClick& InNavigateDelegate, const FSlateHyperlinkRun::FOnGetTooltipText& InToolTipTextDelegate, const FSlateHyperlinkRun::FOnGenerateTooltip& InToolTipDelegate )
+	: FHyperlinkDecorator(InId, InNavigateDelegate, InToolTipTextDelegate, InToolTipDelegate)
+{
+}
+
+TSharedRef< ISlateRun > FTutorialHyperlinkDecorator::Create(const TSharedRef<class FTextLayout>& TextLayout, const FTextRunParseResults& RunParseResult, const FString& OriginalText, const TSharedRef< FString >& InOutModelText, const ISlateStyle* Style)
+{
+	FString StyleName = TEXT("Hyperlink");
+
+	const FTextRange* const MetaDataStyleNameRange = RunParseResult.MetaData.Find( TEXT("style") );
+	if ( MetaDataStyleNameRange != NULL )
+	{
+		const FString MetaDataStyleName = OriginalText.Mid(MetaDataStyleNameRange->BeginIndex, MetaDataStyleNameRange->EndIndex - MetaDataStyleNameRange->BeginIndex);
+		StyleName = *MetaDataStyleName;
+	}
+
+	FTextRange ModelRange;
+	ModelRange.BeginIndex = InOutModelText->Len();
+	*InOutModelText += OriginalText.Mid(RunParseResult.ContentRange.BeginIndex, RunParseResult.ContentRange.EndIndex - RunParseResult.ContentRange.BeginIndex);
+	ModelRange.EndIndex = InOutModelText->Len();
+
+	if ( !Style->HasWidgetStyle<FHyperlinkStyle>( FName( *StyleName ) ) )
+	{
+		Style = &FCoreStyle::Get();
+	}
+
+	FRunInfo RunInfo( RunParseResult.Name );
+	for(const TPair<FString, FTextRange>& Pair : RunParseResult.MetaData)
+	{
+		RunInfo.MetaData.Add(Pair.Key, OriginalText.Mid( Pair.Value.BeginIndex, Pair.Value.EndIndex - Pair.Value.BeginIndex));
+	}
+
+	return FTutorialHyperlinkRun::Create( RunInfo, InOutModelText, Style->GetWidgetStyle<FHyperlinkStyle>( FName( *StyleName ) ), NavigateDelegate, ToolTipDelegate, ToolTipTextDelegate, ModelRange );
+}
diff --git a/Engine/Source/Editor/IntroTutorials/Private/TutorialHyperlinkDecorator.h b/Engine/Source/Editor/IntroTutorials/Private/TutorialHyperlinkDecorator.h
new file mode 100644
--- /dev/null
+++ b/Engine/Source/Editor/IntroTutorials/Private/TutorialHyperlinkDecorator.h
@@ -0,0 +1,19 @@
+// Copyright 1998-2014 Epic Games, Inc. All Rights Reserved.
+#pragma once
+
+#include "TextDecorators.h"
+
+class FTutorialHyperlinkDecorator : public FHyperlinkDecorator
+{
+public:
+
+	static TSharedRef< FTutorialHyperlinkDecorator > Create( FString InId, const FSlateHyperlinkRun::FOnClick& InNavigateDelegate, const FSlateHyperlinkRun::FOnGetTooltipText& InToolTipTextDelegate, const FSlateHyperlinkRun::FOnGenerateTooltip& InToolTipDelegate );
+
+public:
+
+	virtual TSharedRef< ISlateRun > Create(const TSharedRef<class FTextLayout>& TextLayout, const FTextRunParseResults& RunParseResult, const FString& OriginalText, const TSharedRef< FString >& InOutModelText, const ISlateStyle* Style) override;
+
+protected:
+
+	FTutorialHyperlinkDecorator( FString InId, const FSlateHyperlinkRun::FOnClick& InNavigateDelegate, const FSlateHyperlinkRun::FOnGetTooltipText& InToolTipTextDelegate, const FSlateHyperlinkRun::FOnGenerateTooltip& InToolTipDelegate );
+};
\ No newline at end of file
diff --git a/Engine/Source/Editor/IntroTutorials/Private/TutorialHyperlinkRun.cpp b/Engine/Source/Editor/IntroTutorials/Private/TutorialHyperlinkRun.cpp
new file mode 100644
--- /dev/null
+++ b/Engine/Source/Editor/IntroTutorials/Private/TutorialHyperlinkRun.cpp
@@ -0,0 +1,121 @@
+// Copyright 1998-2014 Epic Games, Inc. All Rights Reserved.
+
+#include "IntroTutorialsPrivatePCH.h"
+#include "TutorialHyperlinkRun.h"
+
+TSharedRef< FTutorialHyperlinkRun > FTutorialHyperlinkRun::Create( const FRunInfo& InRunInfo, const TSharedRef< const FString >& InText, const FHyperlinkStyle& InStyle, FOnClick NavigateDelegate, FOnGenerateTooltip InTooltipDelegate, FOnGetTooltipText InTooltipTextDelegate )
+{
+	return MakeShareable( new FTutorialHyperlinkRun( InRunInfo, InText, InStyle, NavigateDelegate, InTooltipDelegate, InTooltipTextDelegate ) );
+}
+
+TSharedRef< FTutorialHyperlinkRun > FTutorialHyperlinkRun::Create( const FRunInfo& InRunInfo, const TSharedRef< const FString >& InText, const FHyperlinkStyle& InStyle, FOnClick NavigateDelegate, FOnGenerateTooltip InTooltipDelegate, FOnGetTooltipText InTooltipTextDelegate, const FTextRange& InRange )
+{
+	return MakeShareable( new FTutorialHyperlinkRun( InRunInfo, InText, InStyle, NavigateDelegate, InTooltipDelegate, InTooltipTextDelegate, InRange ) );
+}
+
+FVector2D FTutorialHyperlinkRun::Measure( int32 StartIndex, int32 EndIndex, float Scale ) const 
+{
+	FVector2D Measurement(0.0f, 0.0f);
+	if ( EndIndex - StartIndex == 0 )
+	{
+		Measurement = FVector2D( 0, GetMaxHeight( Scale ) );
+	}
+	else
+	{
+		const TSharedRef< FSlateFontMeasure > FontMeasure = FSlateApplication::Get().GetRenderer()->GetFontMeasureService();
+		Measurement = FontMeasure->Measure( **Text, StartIndex, EndIndex, Style.TextStyle.Font, true, Scale );	
+	}
+
+	if(bIsBrowserLink && EndIndex == Range.EndIndex)
+	{
+		Measurement.X += LinkBrush->ImageSize.X;
+	}
+
+	return Measurement;
+}
+
+TSharedRef< ILayoutBlock > FTutorialHyperlinkRun::CreateBlock( int32 StartIndex, int32 EndIndex, FVector2D Size, const TSharedPtr< IRunRenderer >& Renderer )
+{
+	FText ToolTipText;
+	TSharedPtr<IToolTip> ToolTip;
+	
+	if(TooltipDelegate.IsBound())
+	{
+		ToolTip = TooltipDelegate.Execute(RunInfo.MetaData);
+	}
+	else
+	{
+		const FString* Url = RunInfo.MetaData.Find(TEXT("href"));
+		if(TooltipTextDelegate.IsBound())
+		{
+			ToolTipText = TooltipTextDelegate.Execute(RunInfo.MetaData);
+		}
+		else if(Url != nullptr)
+		{
+			ToolTipText = FText::FromString(*Url);
+		}
+	}
+
+	TSharedRef<SHorizontalBox> Widget = 
+		SNew(SHorizontalBox)
+		.ToolTipText(ToolTipText)
+		.ToolTip(ToolTip)
+		+SHorizontalBox::Slot()
+		.AutoWidth()
+		[
+			SNew( SRichTextHyperlink, ViewModel )
+			.Style( &Style )
+			.Text( FText::FromString( FString( EndIndex - StartIndex, **Text + StartIndex ) ) )
+			.ToolTip( ToolTip )
+			.ToolTipText( ToolTipText )
+			.OnNavigate( this, &FTutorialHyperlinkRun::OnNavigate )
+		];
+
+	if(bIsBrowserLink && EndIndex == Range.EndIndex)
+	{
+		Widget->AddSlot()
+		.AutoWidth()
+		[
+			SNew(SImage)
+			.Image(LinkBrush)
+		];
+	}
+	
+	// We need to do a prepass here as CreateBlock can be called after the main Slate prepass has been run, 
+	// which can result in the hyperlink widget not being correctly setup before it is painted
+	Widget->SlatePrepass();
+
+	Children.Add( Widget );
+
+	return FWidgetLayoutBlock::Create( SharedThis( this ), Widget, FTextRange( StartIndex, EndIndex ), Size, Renderer );
+}
+
+FTutorialHyperlinkRun::FTutorialHyperlinkRun( const FRunInfo& InRunInfo, const TSharedRef< const FString >& InText, const FHyperlinkStyle& InStyle, FOnClick InNavigateDelegate, FOnGenerateTooltip InTooltipDelegate, FOnGetTooltipText InTooltipTextDelegate ) 
+	: FSlateHyperlinkRun(InRunInfo, InText, InStyle, InNavigateDelegate, InTooltipDelegate, InTooltipTextDelegate)
+	, bIsBrowserLink(CheckIsBrowserLink())
+	, LinkBrush(FEditorStyle::Get().GetBrush("Tutorials.Content.ExternalLink"))
+{
+
+}
+
+FTutorialHyperlinkRun::FTutorialHyperlinkRun( const FRunInfo& InRunInfo, const TSharedRef< const FString >& InText, const FHyperlinkStyle& InStyle, FOnClick InNavigateDelegate, FOnGenerateTooltip InTooltipDelegate, FOnGetTooltipText InTooltipTextDelegate, const FTextRange& InRange ) 
+	: FSlateHyperlinkRun(InRunInfo, InText, InStyle, InNavigateDelegate, InTooltipDelegate, InTooltipTextDelegate,InRange)
+	, bIsBrowserLink(CheckIsBrowserLink())
+	, LinkBrush(FEditorStyle::Get().GetBrush("Tutorials.Content.ExternalLink"))
+{
+
+}
+
+FTutorialHyperlinkRun::FTutorialHyperlinkRun( const FTutorialHyperlinkRun& Run ) 
+	: FSlateHyperlinkRun(Run)
+	, bIsBrowserLink(Run.bIsBrowserLink)
+	, LinkBrush(Run.LinkBrush)
+{
+
+}
+
+bool FTutorialHyperlinkRun::CheckIsBrowserLink() const
+{
+	const FString* IdMetaData = RunInfo.MetaData.Find(TEXT("id"));
+	return (IdMetaData != nullptr && (*IdMetaData == TEXT("browser") || *IdMetaData == TEXT("udn")));
+}
\ No newline at end of file
diff --git a/Engine/Source/Editor/IntroTutorials/Private/TutorialHyperlinkRun.h b/Engine/Source/Editor/IntroTutorials/Private/TutorialHyperlinkRun.h
new file mode 100644
--- /dev/null
+++ b/Engine/Source/Editor/IntroTutorials/Private/TutorialHyperlinkRun.h
@@ -0,0 +1,38 @@
+// Copyright 1998-2014 Epic Games, Inc. All Rights Reserved.
+#pragma once
+
+class FTutorialHyperlinkRun : public FSlateHyperlinkRun
+{
+public:
+
+	static TSharedRef< FTutorialHyperlinkRun > Create( const FRunInfo& InRunInfo, const TSharedRef< const FString >& InText, const FHyperlinkStyle& InStyle, FOnClick NavigateDelegate, FOnGenerateTooltip InTooltipDelegate, FOnGetTooltipText InTooltipTextDelegate );
+																																	 
+	static TSharedRef< FTutorialHyperlinkRun > Create( const FRunInfo& InRunInfo, const TSharedRef< const FString >& InText, const FHyperlinkStyle& InStyle, FOnClick NavigateDelegate, FOnGenerateTooltip InTooltipDelegate, FOnGetTooltipText InTooltipTextDelegate, const FTextRange& InRange );
+
+public:
+
+	virtual ~FTutorialHyperlinkRun() {}
+
+	/** FSlateHyperlinkRun implementation */
+	virtual FVector2D Measure( int32 StartIndex, int32 EndIndex, float Scale ) const override;
+	virtual TSharedRef< ILayoutBlock > CreateBlock( int32 StartIndex, int32 EndIndex, FVector2D Size, const TSharedPtr< IRunRenderer >& Renderer ) override;
+
+private:
+
+	FTutorialHyperlinkRun( const FRunInfo& InRunInfo, const TSharedRef< const FString >& InText, const FHyperlinkStyle& InStyle, FOnClick InNavigateDelegate, FOnGenerateTooltip InTooltipDelegate, FOnGetTooltipText InTooltipTextDelegate );
+																										 
+	FTutorialHyperlinkRun( const FRunInfo& InRunInfo, const TSharedRef< const FString >& InText, const FHyperlinkStyle& InStyle, FOnClick InNavigateDelegate, FOnGenerateTooltip InTooltipDelegate, FOnGetTooltipText InTooltipTextDelegate, const FTextRange& InRange );
+
+	FTutorialHyperlinkRun( const FTutorialHyperlinkRun& Run );
+
+private:
+	// check whether our metadata makes us a browser link
+	bool CheckIsBrowserLink() const;
+
+private:
+	/** Whether we are an external (browser) link or not */
+	bool bIsBrowserLink;
+
+	/** The brush we use to display external links */
+	const FSlateBrush* LinkBrush;
+};
diff --git a/Engine/Source/Editor/IntroTutorials/Private/TutorialText.cpp b/Engine/Source/Editor/IntroTutorials/Private/TutorialText.cpp
--- a/Engine/Source/Editor/IntroTutorials/Private/TutorialText.cpp
+++ b/Engine/Source/Editor/IntroTutorials/Private/TutorialText.cpp
@@ -12,6 +12,8 @@
 #include "TutorialImageDecorator.h"
 #include "EngineAnalytics.h"
 #include "Runtime/Analytics/Analytics/Public/Interfaces/IAnalyticsProvider.h"
+#include "TutorialHyperlinkDecorator.h"
+#include "TextDecorators.h"
 
 #define LOCTEXT_NAMESPACE "TutorialText"
 
@@ -298,13 +300,23 @@ static FText OnGetTutorialTooltipText(const FSlateHyperlinkRun::FMetadata& Metad
 	return LOCTEXT("InvalidTutorialLink", "Invalid Tutorial Link");
 }
 
-void FTutorialText::GetRichTextDecorators(TArray< TSharedRef< class ITextDecorator > >& OutDecorators)
+void FTutorialText::GetRichTextDecorators(bool bForEditing, TArray< TSharedRef< class ITextDecorator > >& OutDecorators)
 {
 	Initialize();
 
-	for(const auto& HyperlinkDesc : HyperlinkDescs)
+	if(bForEditing)
 	{
-		OutDecorators.Add(FHyperlinkDecorator::Create(HyperlinkDesc->Id, HyperlinkDesc->OnClickedDelegate, HyperlinkDesc->TooltipTextDelegate, HyperlinkDesc->TooltipDelegate));
+		for(const auto& HyperlinkDesc : HyperlinkDescs)
+		{
+			OutDecorators.Add(FHyperlinkDecorator::Create(HyperlinkDesc->Id, HyperlinkDesc->OnClickedDelegate, HyperlinkDesc->TooltipTextDelegate, HyperlinkDesc->TooltipDelegate));
+		}
+	}
+	else
+	{
+		for(const auto& HyperlinkDesc : HyperlinkDescs)
+		{
+			OutDecorators.Add(FTutorialHyperlinkDecorator::Create(HyperlinkDesc->Id, HyperlinkDesc->OnClickedDelegate, HyperlinkDesc->TooltipTextDelegate, HyperlinkDesc->TooltipDelegate));
+		}	
 	}
 
 	OutDecorators.Add(FTextStyleDecorator::Create());
diff --git a/Engine/Source/Editor/IntroTutorials/Private/TutorialText.h b/Engine/Source/Editor/IntroTutorials/Private/TutorialText.h
--- a/Engine/Source/Editor/IntroTutorials/Private/TutorialText.h
+++ b/Engine/Source/Editor/IntroTutorials/Private/TutorialText.h
@@ -168,7 +168,7 @@ struct FTextStyleAndName
 struct FTutorialText
 {
 public:
-	static void GetRichTextDecorators(TArray<TSharedRef<class ITextDecorator>>& OutDecorators);
+	static void GetRichTextDecorators(bool bForEditing, TArray<TSharedRef<class ITextDecorator>>& OutDecorators);
 
 	static const TArray<TSharedPtr<FHyperlinkTypeDesc>>& GetHyperLinkDescs();
 
diff --git a/Engine/Source/Runtime/Slate/Public/Framework/Text/SlateHyperlinkRun.h b/Engine/Source/Runtime/Slate/Public/Framework/Text/SlateHyperlinkRun.h
--- a/Engine/Source/Runtime/Slate/Public/Framework/Text/SlateHyperlinkRun.h
+++ b/Engine/Source/Runtime/Slate/Public/Framework/Text/SlateHyperlinkRun.h
@@ -77,7 +77,7 @@ public:
 
 	virtual ERunAttributes GetRunAttributes() const override;
 
-private:
+protected:
 
 	FSlateHyperlinkRun( const FRunInfo& InRunInfo, const TSharedRef< const FString >& InText, const FHyperlinkStyle& InStyle, FOnClick InNavigateDelegate, FOnGenerateTooltip InTooltipDelegate, FOnGetTooltipText InTooltipTextDelegate );
 																										 
@@ -85,11 +85,11 @@ private:
 
 	FSlateHyperlinkRun( const FSlateHyperlinkRun& Run );
 
-private:
+protected:
 
 	void OnNavigate();
 
-private:
+protected:
 
 	FRunInfo RunInfo;
 	TSharedRef< const FString > Text;
diff --git a/Engine/Source/Runtime/Slate/Public/Framework/Text/TextDecorators.h b/Engine/Source/Runtime/Slate/Public/Framework/Text/TextDecorators.h
--- a/Engine/Source/Runtime/Slate/Public/Framework/Text/TextDecorators.h
+++ b/Engine/Source/Runtime/Slate/Public/Framework/Text/TextDecorators.h
@@ -76,7 +76,7 @@ protected:
 
 	FSlateHyperlinkRun::FOnClick NavigateDelegate;
 
-private:
+protected:
 	FString Id;
 	FSlateHyperlinkRun::FOnGetTooltipText ToolTipTextDelegate;
 	FSlateHyperlinkRun::FOnGenerateTooltip ToolTipDelegate;
