Cleaned up unused old tutorial system files

diff --git a/Engine/Source/Editor/IntroTutorials/Classes/K2Node_TutorialExcerptComplete.h b/Engine/Source/Editor/IntroTutorials/Classes/K2Node_TutorialExcerptComplete.h
deleted file mode 100644
diff --git a/Engine/Source/Editor/IntroTutorials/Private/InteractiveTutorial.cpp b/Engine/Source/Editor/IntroTutorials/Private/InteractiveTutorial.cpp
deleted file mode 100644
diff --git a/Engine/Source/Editor/IntroTutorials/Private/InteractiveTutorial.h b/Engine/Source/Editor/IntroTutorials/Private/InteractiveTutorial.h
deleted file mode 100644
diff --git a/Engine/Source/Editor/IntroTutorials/Private/InteractiveTutorials.cpp b/Engine/Source/Editor/IntroTutorials/Private/InteractiveTutorials.cpp
deleted file mode 100644
diff --git a/Engine/Source/Editor/IntroTutorials/Private/InteractiveTutorials.h b/Engine/Source/Editor/IntroTutorials/Private/InteractiveTutorials.h
deleted file mode 100644
diff --git a/Engine/Source/Editor/IntroTutorials/Private/InteractivityData.cpp b/Engine/Source/Editor/IntroTutorials/Private/InteractivityData.cpp
deleted file mode 100644
diff --git a/Engine/Source/Editor/IntroTutorials/Private/InteractivityData.h b/Engine/Source/Editor/IntroTutorials/Private/InteractivityData.h
deleted file mode 100644
diff --git a/Engine/Source/Editor/IntroTutorials/Private/IntroTutorials.cpp b/Engine/Source/Editor/IntroTutorials/Private/IntroTutorials.cpp
--- a/Engine/Source/Editor/IntroTutorials/Private/IntroTutorials.cpp
+++ b/Engine/Source/Editor/IntroTutorials/Private/IntroTutorials.cpp
@@ -1,7 +1,6 @@
 // Copyright 1998-2014 Epic Games, Inc. All Rights Reserved.
 
 #include "IntroTutorialsPrivatePCH.h"
-#include "SIntroTutorials.h"
 #include "Editor/WorkspaceMenuStructure/Public/WorkspaceMenuStructureModule.h"
 #include "Editor/MainFrame/Public/MainFrame.h"
 #include "Editor/LevelEditor/Public/LevelEditor.h"
diff --git a/Engine/Source/Editor/IntroTutorials/Private/K2Node_TutorialExcerptComplete.cpp b/Engine/Source/Editor/IntroTutorials/Private/K2Node_TutorialExcerptComplete.cpp
deleted file mode 100644
diff --git a/Engine/Source/Editor/IntroTutorials/Private/SIntroTutorials.cpp b/Engine/Source/Editor/IntroTutorials/Private/SIntroTutorials.cpp
deleted file mode 100644
diff --git a/Engine/Source/Editor/IntroTutorials/Private/SIntroTutorials.h b/Engine/Source/Editor/IntroTutorials/Private/SIntroTutorials.h
deleted file mode 100644
