Moving tutorial browser to an editor nomad tab

Tweaked styles and layout to better suit a tab.
Also fixed browser not showing some tutorials if all assets are not yet loaded.
Progress is now updated periodically as the browser can be now open while a tutorial is in progress.

TTP# 347956 - TUTORIALS: Selecting More Tutorials can hide Tutorial Browser

diff --git a/Engine/Source/Editor/EditorStyle/Private/SlateEditorStyle.cpp b/Engine/Source/Editor/EditorStyle/Private/SlateEditorStyle.cpp
--- a/Engine/Source/Editor/EditorStyle/Private/SlateEditorStyle.cpp
+++ b/Engine/Source/Editor/EditorStyle/Private/SlateEditorStyle.cpp
@@ -2482,22 +2482,37 @@ void FSlateEditorStyle::FStyle::SetupTutorialStyles()
 			.SetHighlightColor(TutorialSelectionColor);
 
 		Set( "Tutorials.Border", new BOX_BRUSH( "Tutorials/OverlayFrame", FMargin(18.0f/64.0f), FLinearColor(1.0f, 1.0f, 1.0f, 1.0f) ) );
-		Set( "Tutorials.Browser.WelcomeHeader", FTextBlockStyle(TutorialHeaderText)
+
+		const FTextBlockStyle TutorialBrowserText = FTextBlockStyle(TutorialText)
+			.SetColorAndOpacity(FSlateColor::UseForeground())
+			.SetHighlightColor(TutorialSelectionColor);
+
+		Set( "Tutorials.Browser.Text", TutorialBrowserText );
+
+		Set( "Tutorials.Browser.WelcomeHeader", FTextBlockStyle(TutorialBrowserText)
 			.SetFontSize(20));
 
-		Set( "Tutorials.Browser.SummaryHeader", FTextBlockStyle(TutorialHeaderText)
+		Set( "Tutorials.Browser.SummaryHeader", FTextBlockStyle(TutorialBrowserText)
 			.SetFontSize(16));
 
-		Set( "Tutorials.Browser.SummaryText", FTextBlockStyle(TutorialText)
+		Set( "Tutorials.Browser.SummaryText", FTextBlockStyle(TutorialBrowserText)
 			.SetFontSize(10));
 
 		Set( "Tutorials.Browser.HighlightTextColor", TutorialSelectionColor );
 
 		Set( "Tutorials.Browser.Button", FButtonStyle()
-			.SetNormal( BOX_BRUSH( "Common/ButtonHoverHint", FMargin(4/16.0f), FLinearColor(1,1,1,1) ) )
-			.SetHovered( BOX_BRUSH( "Common/ButtonHoverHint", FMargin(4/16.0f), FLinearColor(0.9f,0.9f,0.9f,1) ) )
-			.SetPressed( BOX_BRUSH( "Common/ButtonHoverHint", FMargin(4/16.0f), FLinearColor(1,1,1,1) ) )
-			.SetDisabled( BOX_BRUSH( "Common/ButtonHoverHint", FMargin(4/16.0f), FLinearColor(1,1,1,1) ) )
+			.SetNormal( BOX_BRUSH( "Common/ButtonHoverHint", FMargin(4/16.0f), FLinearColor(0.05f,0.05f,0.05f,1) ) )
+			.SetHovered( BOX_BRUSH( "Common/ButtonHoverHint", FMargin(4/16.0f), FLinearColor(0.07f,0.07f,0.07f,1) ) )
+			.SetPressed( BOX_BRUSH( "Common/ButtonHoverHint", FMargin(4/16.0f), FLinearColor(0.08f,0.08f,0.08f,1) ) )
+			.SetDisabled( BOX_BRUSH( "Common/ButtonHoverHint", FMargin(4/16.0f), FLinearColor(0.1f,0.1f,0.1f,1) ) )
+			.SetNormalPadding( FMargin(0,0,0,1))
+			.SetPressedPadding( FMargin(0,1,0,0)));
+
+		Set( "Tutorials.Browser.BackButton", FButtonStyle()
+			.SetNormal( BOX_BRUSH( "Common/ButtonHoverHint", FMargin(4/16.0f), FLinearColor(1.0f,1.0f,1.0f,0.0f) ) )
+			.SetHovered( BOX_BRUSH( "Common/ButtonHoverHint", FMargin(4/16.0f), FLinearColor(1.0f,1.0f,1.0f,0.05f) ) )
+			.SetPressed( BOX_BRUSH( "Common/ButtonHoverHint", FMargin(4/16.0f), FLinearColor(1.0f,1.0f,1.0f,0.05f) ) )
+			.SetDisabled( BOX_BRUSH( "Common/ButtonHoverHint", FMargin(4/16.0f), FLinearColor(1.0f,1.0f,1.0f,0.05f) ) )
 			.SetNormalPadding( FMargin(0,0,0,1))
 			.SetPressedPadding( FMargin(0,1,0,0)));
 
@@ -2530,17 +2545,17 @@ void FSlateEditorStyle::FStyle::SetupTutorialStyles()
 		Set( "Tutorials.Content.Color", FLinearColor(1.0f,1.0f,1.0f,0.9f) );
 		Set( "Tutorials.Content.Color.Hovered", FLinearColor(1.0f,1.0f,1.0f,1.0f) );
 
-		Set( "Tutorials.Browser.CategoryArrow", new IMAGE_BRUSH( "Tutorials/BrowserCategoryArrow", FVector2D(24.0f, 24.0f), FLinearColor(0.0f, 0.0f, 0.0f, 1.0f) ) );
+		Set( "Tutorials.Browser.CategoryArrow", new IMAGE_BRUSH( "Tutorials/BrowserCategoryArrow", FVector2D(24.0f, 24.0f), FSlateColor::UseForeground() ) );
 		Set( "Tutorials.Browser.DefaultTutorialIcon", new IMAGE_BRUSH( "Tutorials/DefaultTutorialIcon_40x", FVector2D(40.0f, 40.0f), FLinearColor::White ) );
 		Set( "Tutorials.Browser.DefaultCategoryIcon", new IMAGE_BRUSH( "Tutorials/DefaultCategoryIcon_40x", FVector2D(40.0f, 40.0f), FLinearColor::White ) );
 
-		Set( "Tutorials.Browser.BackButton", new IMAGE_BRUSH( "Tutorials/BrowserBack", FVector2D(32.0f, 32.0f), FLinearColor(0.0f, 0.0f, 0.0f, 1.0f) ) );
-		Set( "Tutorials.Browser.RestartButton", new IMAGE_BRUSH( "Tutorials/BrowserRestart", FVector2D(16.0f, 16.0f), FLinearColor(0.0f, 0.0f, 0.0f, 1.0f) ) );
+		Set( "Tutorials.Browser.BackButton.Image", new IMAGE_BRUSH( "Tutorials/BrowserBack", FVector2D(32.0f, 32.0f), FLinearColor(1.0f, 1.0f, 1.0f, 1.0f) ) );
+		Set( "Tutorials.Browser.RestartButton", new IMAGE_BRUSH( "Tutorials/BrowserRestart", FVector2D(16.0f, 16.0f), FLinearColor(1.0f, 1.0f, 1.0f, 1.0f) ) );
 
 		Set( "Tutorials.Browser.Completed", new IMAGE_BRUSH( "Tutorials/TutorialCompleted", Icon32x32 ) );
 
-		Set( "Tutorials.Browser.Breadcrumb", new IMAGE_BRUSH( "Tutorials/Breadcrumb", Icon8x8, FLinearColor(0.0f, 0.0f, 0.0f, 1.0f) ) );
-		Set( "Tutorials.Browser.PathText", FTextBlockStyle(TutorialText)
+		Set( "Tutorials.Browser.Breadcrumb", new IMAGE_BRUSH( "Tutorials/Breadcrumb", Icon8x8, FLinearColor::White ) );
+		Set( "Tutorials.Browser.PathText", FTextBlockStyle(TutorialBrowserText)
 			.SetFontSize(9));
 
 		Set( "Tutorials.Navigation.Button", FButtonStyle()
@@ -2579,7 +2594,8 @@ void FSlateEditorStyle::FStyle::SetupTutorialStyles()
 			.SetFontSize(16));
 
 		// UDN documentation styles
-		Set("Tutorials.Content",         FTextBlockStyle(TutorialText));
+		Set("Tutorials.Content", FTextBlockStyle(TutorialText)
+			.SetColorAndOpacity(FSlateColor::UseForeground()));
 		Set("Tutorials.Hyperlink.Text",  FTextBlockStyle(DocumentationHyperlinkText));
 		Set("Tutorials.NumberedContent", FTextBlockStyle(TutorialText));
 		Set("Tutorials.BoldContent",     FTextBlockStyle(TutorialText)
diff --git a/Engine/Source/Editor/IntroTutorials/Private/IntroTutorials.cpp b/Engine/Source/Editor/IntroTutorials/Private/IntroTutorials.cpp
--- a/Engine/Source/Editor/IntroTutorials/Private/IntroTutorials.cpp
+++ b/Engine/Source/Editor/IntroTutorials/Private/IntroTutorials.cpp
@@ -144,6 +144,9 @@ void FIntroTutorials::StartupModule()
 		ContentIntroCurveAsset->AddToRoot();
 		ContentIntroCurve = ContentIntroCurveAsset;
 	}
+
+	FGlobalTabmanager::Get()->RegisterNomadTabSpawner(TEXT("TutorialsBrowser"), FOnSpawnTab::CreateRaw(this, &FIntroTutorials::SpawnTutorialsBrowserTab))
+		.SetMenuType(ETabSpawnerMenuType::Hide);
 }
 
 void FIntroTutorials::ShutdownModule()
@@ -189,6 +192,8 @@ void FIntroTutorials::ShutdownModule()
 		ContentIntroCurve.Get()->RemoveFromRoot();
 		ContentIntroCurve = nullptr;
 	}
+
+	FGlobalTabmanager::Get()->UnregisterNomadTabSpawner(TEXT("TutorialsBrowser"));
 }
 
 void FIntroTutorials::AddSummonTutorialsMenuExtension(FMenuBuilder& MenuBuilder)
@@ -367,11 +372,12 @@ FOnIsPicking& FIntroTutorials::OnIsPicking()
 	return OnIsPickingDelegate;
 }
 
-void FIntroTutorials::SummonTutorialBrowser(TWeakPtr<SWindow> InWindow, const FString& InFilter)
+void FIntroTutorials::SummonTutorialBrowser()
 {
 	if(TutorialRoot.IsValid())
 	{
-		TutorialRoot->SummonTutorialBrowser(InWindow, InFilter);
+		FLevelEditorModule& LevelEditorModule = FModuleManager::GetModuleChecked<FLevelEditorModule>( TEXT("LevelEditor") );
+		LevelEditorModule.GetLevelEditorTabManager()->InvokeTab(FTabId("TutorialsBrowser"));
 	}
 }
 
@@ -444,4 +450,21 @@ float FIntroTutorials::GetIntroCurveValue(float InTime)
 	return 1.0f;
 }
 
+TSharedRef<SDockTab> FIntroTutorials::SpawnTutorialsBrowserTab(const FSpawnTabArgs& SpawnTabArgs)
+{
+	TAttribute<FText> Label = LOCTEXT("TutorialsBrowserTabLabel", "Tutorials");
+
+	TSharedRef<SDockTab> NewTab = SNew(SDockTab)
+		.TabRole(ETabRole::NomadTab)
+		.Label(Label)
+		.ToolTip(IDocumentation::Get()->CreateToolTip(Label, nullptr, "Shared/TutorialsBrowser", "Tab"));	
+
+	TSharedRef<STutorialsBrowser> TutorialsBrowser = SNew(STutorialsBrowser)
+		.OnLaunchTutorial(FOnLaunchTutorial::CreateRaw(this, &FIntroTutorials::LaunchTutorial));
+
+	NewTab->SetContent(TutorialsBrowser);
+
+	return NewTab;
+}
+
 #undef LOCTEXT_NAMESPACE
\ No newline at end of file
diff --git a/Engine/Source/Editor/IntroTutorials/Private/IntroTutorials.h b/Engine/Source/Editor/IntroTutorials/Private/IntroTutorials.h
--- a/Engine/Source/Editor/IntroTutorials/Private/IntroTutorials.h
+++ b/Engine/Source/Editor/IntroTutorials/Private/IntroTutorials.h
@@ -24,7 +24,7 @@ public:
 
 	float GetIntroCurveValue(float InTime);
 
-	void SummonTutorialBrowser(TWeakPtr<SWindow> InWindow = nullptr, const FString& InFilter = TEXT(""));
+	void SummonTutorialBrowser();
 
 	static FString AnalyticsEventNameFromTutorial(const FString& BaseEventName, UEditorTutorial* Tutorial);
 
@@ -80,6 +80,9 @@ private:
 	/** Internal helper to launch a tutorial from a path */
 	void LaunchTutorialByName(const FString& InAssetPath, bool bRestart = true, TWeakPtr<SWindow> InNavigationWindow = nullptr, FSimpleDelegate OnTutorialClosed = FSimpleDelegate(), FSimpleDelegate OnTutorialExited = FSimpleDelegate());
 
+	/** Spawn a tab for browsing tutorials */
+	TSharedRef<SDockTab> SpawnTutorialsBrowserTab(const FSpawnTabArgs& SpawnTabArgs);
+
 public:
 
 	// IIntroTutorials interface
diff --git a/Engine/Source/Editor/IntroTutorials/Private/SEditorTutorials.cpp b/Engine/Source/Editor/IntroTutorials/Private/SEditorTutorials.cpp
--- a/Engine/Source/Editor/IntroTutorials/Private/SEditorTutorials.cpp
+++ b/Engine/Source/Editor/IntroTutorials/Private/SEditorTutorials.cpp
@@ -11,7 +11,6 @@
 
 void SEditorTutorials::Construct(const FArguments& InArgs)
 {
-	bBrowserVisible = false;
 	bIsNavigationWindow = false;
 	ParentWindow = InArgs._ParentWindow;
 	OnNextClicked = InArgs._OnNextClicked;
@@ -23,62 +22,11 @@ void SEditorTutorials::Construct(const FArguments& InArgs)
 	OnWidgetWasDrawn = InArgs._OnWidgetWasDrawn;
 	OnWasWidgetDrawn = InArgs._OnWasWidgetDrawn;
 
-	TutorialHome = SNew(STutorialsBrowser)
-		.Visibility(this, &SEditorTutorials::GetBrowserVisibility)
-		.OnClosed(FSimpleDelegate::CreateSP(this, &SEditorTutorials::HandleCloseClicked))
-		.OnLaunchTutorial(InArgs._OnLaunchTutorial)
-		.ParentWindow(InArgs._ParentWindow);
-
-	NavigationWidget = SNew(STutorialNavigation)
-		.Visibility(this, &SEditorTutorials::GetNavigationVisibility)
-		.OnBackClicked(FSimpleDelegate::CreateSP(this, &SEditorTutorials::HandleBackClicked))
-		.OnHomeClicked(FSimpleDelegate::CreateSP(this, &SEditorTutorials::HandleHomeClicked))
-		.OnNextClicked(FSimpleDelegate::CreateSP(this, &SEditorTutorials::HandleNextClicked))
-		.IsBackEnabled(this, &SEditorTutorials::IsBackButtonEnabled)
-		.IsHomeEnabled(this, &SEditorTutorials::IsHomeButtonEnabled)
-		.IsNextEnabled(this, &SEditorTutorials::IsNextButtonEnabled)
-		.OnGetProgress(this, &SEditorTutorials::GetProgress)
-		.ParentWindow(InArgs._ParentWindow);
-
 	ContentBox = SNew(SHorizontalBox);
 
 	ChildSlot
 	[
-		SNew(SOverlay)
-		+SOverlay::Slot()
-		[
-			SNew(SVerticalBox)
-			+SVerticalBox::Slot()
-			.VAlign(VAlign_Center)
-			.HAlign(HAlign_Center)
-			[
-				TutorialHome.ToSharedRef()
-			]
-		]
-		+SOverlay::Slot()
-		.VAlign(VAlign_Fill)
-		.HAlign(HAlign_Fill)
-		[
-			SNew(SVerticalBox)
-			+SVerticalBox::Slot()
-			.FillHeight(0.75f)
-			[
-				SNew( SSpacer )
-			]
-			+SVerticalBox::Slot()
-			.FillHeight(0.25f)
-			.VAlign(VAlign_Center)
-			.HAlign(HAlign_Center)
-			[
-				NavigationWidget.ToSharedRef()
-			]
-		]
-		+SOverlay::Slot()
-		.VAlign(VAlign_Fill)
-		.HAlign(HAlign_Fill)
-		[
-			ContentBox.ToSharedRef()
-		]
+		ContentBox.ToSharedRef()
 	];
 }
 
@@ -91,20 +39,9 @@ void SEditorTutorials::LaunchTutorial(bool bInIsNavigationWindow, FSimpleDelegat
 	RebuildCurrentContent();
 }
 
-
-void SEditorTutorials::ShowBrowser(const FString& InFilter)
-{
-	TutorialHome->SetFilter(InFilter);
-	OnHomeClicked.ExecuteIfBound();
-	ContentBox->ClearChildren();
-	OnTutorialExited.ExecuteIfBound();
-	bBrowserVisible = true;
-}
-
 void SEditorTutorials::HideContent()
 {
 	HandleHomeClicked();
-	bBrowserVisible = false;
 	bIsNavigationWindow = false;
 }
 
@@ -115,7 +52,7 @@ bool SEditorTutorials::IsNavigationVisible() const
 
 EVisibility SEditorTutorials::GetBrowserVisibility() const
 {
-	return bBrowserVisible && OnGetCurrentTutorial.Execute() == nullptr ? EVisibility::Visible : EVisibility::Collapsed;
+	return OnGetCurrentTutorial.Execute() == nullptr ? EVisibility::Visible : EVisibility::Collapsed;
 }
 
 EVisibility SEditorTutorials::GetNavigationVisibility() const
@@ -128,8 +65,8 @@ void SEditorTutorials::HandleCloseClicked()
 	OnCloseClicked.ExecuteIfBound();
 	OnTutorialClosed.ExecuteIfBound();
 
-	HandleHomeClicked();
-	bBrowserVisible = false;
+	ContentBox->ClearChildren();
+	OnTutorialExited.ExecuteIfBound();
 }
 
 void SEditorTutorials::HandleBackClicked()
diff --git a/Engine/Source/Editor/IntroTutorials/Private/SEditorTutorials.h b/Engine/Source/Editor/IntroTutorials/Private/SEditorTutorials.h
--- a/Engine/Source/Editor/IntroTutorials/Private/SEditorTutorials.h
+++ b/Engine/Source/Editor/IntroTutorials/Private/SEditorTutorials.h
@@ -42,9 +42,6 @@ class SEditorTutorials : public SCompoundWidget
 	/** Launch a tutorial - will interrogate parent to get the tutorial data to display */
 	void LaunchTutorial(bool bInIsNavigationWindow, FSimpleDelegate InOnTutorialClosed, FSimpleDelegate InOnTutorialExited);
 
-	/** Show the tutorials browser in this window */
-	void ShowBrowser(const FString& InFilter);
-
 	/** hide the tutorials browser in this window */
 	void HideContent();
 
@@ -89,12 +86,6 @@ private:
 	float GetProgress() const;
 
 private:
-	/** The root widget of the tutorial home page */
-	TSharedPtr<STutorialsBrowser> TutorialHome;
-
-	/** Navigation widget */
-	TSharedPtr<SWidget> NavigationWidget;
-
 	/** Box that contains varied content for current tutorial */
 	TSharedPtr<SHorizontalBox> ContentBox;
 
diff --git a/Engine/Source/Editor/IntroTutorials/Private/STutorialButton.cpp b/Engine/Source/Editor/IntroTutorials/Private/STutorialButton.cpp
--- a/Engine/Source/Editor/IntroTutorials/Private/STutorialButton.cpp
+++ b/Engine/Source/Editor/IntroTutorials/Private/STutorialButton.cpp
@@ -150,7 +150,7 @@ FReply STutorialButton::HandleButtonClicked()
 	FIntroTutorials& IntroTutorials = FModuleManager::GetModuleChecked<FIntroTutorials>(TEXT("IntroTutorials"));
 	if(ShouldLaunchBrowser())
 	{
-		IntroTutorials.SummonTutorialBrowser(nullptr, CachedBrowserFilter);
+		IntroTutorials.SummonTutorialBrowser();
 	}
 	else if (CachedLaunchTutorial != nullptr)
 	{
@@ -247,7 +247,7 @@ void STutorialButton::LaunchBrowser()
 	RefreshStatus();
 
 	FIntroTutorials& IntroTutorials = FModuleManager::GetModuleChecked<FIntroTutorials>(TEXT("IntroTutorials"));
-	IntroTutorials.SummonTutorialBrowser(nullptr, CachedBrowserFilter);
+	IntroTutorials.SummonTutorialBrowser();
 }
 
 bool STutorialButton::ShouldLaunchBrowser() const
diff --git a/Engine/Source/Editor/IntroTutorials/Private/STutorialContent.cpp b/Engine/Source/Editor/IntroTutorials/Private/STutorialContent.cpp
--- a/Engine/Source/Editor/IntroTutorials/Private/STutorialContent.cpp
+++ b/Engine/Source/Editor/IntroTutorials/Private/STutorialContent.cpp
@@ -98,7 +98,7 @@ void STutorialContent::Construct(const FArguments& InArgs, UEditorTutorial* InTu
 							.MaxWidth(600.0f)
 							.VAlign(VAlign_Center)
 							[
-								GenerateContentWidget(InContent, InArgs._WrapTextAt, DocumentationPage)
+								GenerateContentWidget(InContent, DocumentationPage, TAttribute<FText>(), false, InArgs._WrapTextAt)
 							]
 						]
 					]
@@ -281,7 +281,7 @@ static TSharedRef<SWidget> GetStageTitle(const FExcerpt& InExcerpt, int32 InCurr
 	return SNullWidget::NullWidget;
 }
 
-TSharedRef<SWidget> STutorialContent::GenerateContentWidget(const FTutorialContent& InContent, float WrapTextAt, TSharedPtr<IDocumentationPage>& OutDocumentationPage, const TAttribute<FText>& InHighlightText)
+TSharedRef<SWidget> STutorialContent::GenerateContentWidget(const FTutorialContent& InContent, TSharedPtr<IDocumentationPage>& OutDocumentationPage, const TAttribute<FText>& InHighlightText, bool bAutoWrapText, float WrapTextAt)
 {
 	// Style for the documentation
 	static FDocumentationStyle DocumentationStyle;
@@ -301,13 +301,20 @@ TSharedRef<SWidget> STutorialContent::GenerateContentWidget(const FTutorialConte
 	{
 	case ETutorialContent::Text:
 		{
-			return SNew(STextBlock)
+			TSharedRef<STextBlock> TextBlock = SNew(STextBlock)
 				.Visibility(EVisibility::SelfHitTestInvisible)
-				.WrapTextAt(WrapTextAt)
+				.AutoWrapText(bAutoWrapText)
 				.Text(InContent.Text)
 				.TextStyle(FEditorStyle::Get(), "Tutorials.Content")
 				.HighlightText(InHighlightText)
 				.HighlightColor(FEditorStyle::Get().GetColor("Tutorials.Browser.HighlightTextColor"));
+
+			if(!bAutoWrapText)
+			{
+				TextBlock->SetWrapTextAt(WrapTextAt);
+			}
+
+			return TextBlock;
 		}
 
 	case ETutorialContent::UDNExcerpt:
@@ -341,16 +348,19 @@ TSharedRef<SWidget> STutorialContent::GenerateContentWidget(const FTutorialConte
 			const bool bForEditing = false;
 			FTutorialText::GetRichTextDecorators(bForEditing, Decorators);
 
-			return SNew(SRichTextBlock)
+			TSharedRef<SRichTextBlock> TextBlock = SNew(SRichTextBlock)
 					.Visibility(EVisibility::SelfHitTestInvisible)
 					.TextStyle(FEditorStyle::Get(), "Tutorials.Content.Text")
 					.DecoratorStyleSet(&FEditorStyle::Get())
 					.Decorators(Decorators)
 					.Text(InContent.Text)
+					.AutoWrapText(bAutoWrapText)
 					.WrapTextAt(WrapTextAt)
 					.Margin(4)
 					.LineHeightPercentage(1.1f)
 					.HighlightText(InHighlightText);
+
+			return TextBlock;
 		}
 		break;
 	}
diff --git a/Engine/Source/Editor/IntroTutorials/Private/STutorialContent.h b/Engine/Source/Editor/IntroTutorials/Private/STutorialContent.h
--- a/Engine/Source/Editor/IntroTutorials/Private/STutorialContent.h
+++ b/Engine/Source/Editor/IntroTutorials/Private/STutorialContent.h
@@ -77,7 +77,7 @@ class STutorialContent : public SCompoundWidget
 	virtual FReply OnMouseButtonDoubleClick( const FGeometry& InMyGeometry, const FPointerEvent& InMouseEvent ) override;
 
 	/** Helper function to generate widgets from an FTutorialContent struct */
-	static TSharedRef<SWidget> GenerateContentWidget(const FTutorialContent& InContent, float WrapTextAt, TSharedPtr<IDocumentationPage>& OutDocumentationPage, const TAttribute<FText>& InHighlightText = TAttribute<FText>());
+	static TSharedRef<SWidget> GenerateContentWidget(const FTutorialContent& InContent, TSharedPtr<IDocumentationPage>& OutDocumentationPage, const TAttribute<FText>& InHighlightText = TAttribute<FText>(), bool bAutoWrapText = true, float WrapTextAt = 0.0f);
 
 	/** Handle repositioning the widget */
 	FVector2D GetPosition() const;
diff --git a/Engine/Source/Editor/IntroTutorials/Private/STutorialRoot.cpp b/Engine/Source/Editor/IntroTutorials/Private/STutorialRoot.cpp
--- a/Engine/Source/Editor/IntroTutorials/Private/STutorialRoot.cpp
+++ b/Engine/Source/Editor/IntroTutorials/Private/STutorialRoot.cpp
@@ -84,22 +84,6 @@ void STutorialRoot::Tick(const FGeometry& AllottedGeometry, const double InCurre
 	DrawnWidgets.Empty(DrawnWidgets.Max());
 }
 
-void STutorialRoot::SummonTutorialBrowser(TWeakPtr<SWindow> InWindow, const FString& InFilter)
-{
-	// Use main frame if non specified
-	if (!InWindow.IsValid())
-	{
-		IMainFrameModule& MainFrameModule = FModuleManager::LoadModuleChecked<IMainFrameModule>(TEXT("MainFrame"));
-		InWindow = MainFrameModule.GetParentWindow();
-	}
-
-	TWeakPtr<SEditorTutorials>* FoundWidget = TutorialWidgets.Find(InWindow);
-	if(FoundWidget != nullptr && FoundWidget->IsValid())
-	{
-		FoundWidget->Pin()->ShowBrowser(InFilter);
-	}
-}
-
 void STutorialRoot::LaunchTutorial(UEditorTutorial* InTutorial, bool bInRestart, TWeakPtr<SWindow> InNavigationWindow, FSimpleDelegate InOnTutorialClosed, FSimpleDelegate InOnTutorialExited)
 {
 	if(InTutorial != nullptr)
diff --git a/Engine/Source/Editor/IntroTutorials/Private/STutorialRoot.h b/Engine/Source/Editor/IntroTutorials/Private/STutorialRoot.h
--- a/Engine/Source/Editor/IntroTutorials/Private/STutorialRoot.h
+++ b/Engine/Source/Editor/IntroTutorials/Private/STutorialRoot.h
@@ -33,9 +33,6 @@ public:
 	/** Close all tutorial content */
 	void CloseAllTutorialContent();
 
-	/** Summon the browser widget for the specified window */
-	void SummonTutorialBrowser(TWeakPtr<SWindow> InWindow = nullptr, const FString& InFilter = TEXT(""));
-
 	/** Reload tutorials that we know about */
 	void ReloadTutorials();
 
diff --git a/Engine/Source/Editor/IntroTutorials/Private/STutorialsBrowser.cpp b/Engine/Source/Editor/IntroTutorials/Private/STutorialsBrowser.cpp
--- a/Engine/Source/Editor/IntroTutorials/Private/STutorialsBrowser.cpp
+++ b/Engine/Source/Editor/IntroTutorials/Private/STutorialsBrowser.cpp
@@ -17,6 +17,13 @@ class FTutorialListEntry_Tutorial;
 
 DECLARE_DELEGATE_OneParam(FOnCategorySelected, const FString& /* InCategory */);
 
+namespace TutorialBrowserConstants
+{
+	const float RefreshTimerInterval = 0.25f;
+
+	const float ProgressUpdateInterval = 0.5f;
+}
+
 class FTutorialListEntry_Category : public ITutorialListEntry, public TSharedFromThis<FTutorialListEntry_Category>
 {
 public:
@@ -76,65 +83,69 @@ public:
 	{
 		return SNew(STableRow<TSharedPtr<ITutorialListEntry>>, OwnerTable)
 		[
-			SNew(SButton)
-			.OnClicked(this, &FTutorialListEntry_Category::OnClicked)
-			.ButtonStyle(&FEditorStyle::Get().GetWidgetStyle<FButtonStyle>("Tutorials.Browser.Button"))
-			.Content()
+			SNew(SBox)
+			.Padding(FMargin(0.0f, 2.0f))
 			[
-				SNew(SHorizontalBox)
-				+SHorizontalBox::Slot()
-				.AutoWidth()
-				.VAlign(VAlign_Center)
-				.HAlign(HAlign_Center)
-				.Padding(8.0f)
+				SNew(SButton)
+				.OnClicked(this, &FTutorialListEntry_Category::OnClicked)
+				.ButtonStyle(&FEditorStyle::Get().GetWidgetStyle<FButtonStyle>("Tutorials.Browser.Button"))
+				.ForegroundColor(FSlateColor::UseForeground())
+				.Content()
 				[
-					SNew(SBox)
-					.WidthOverride(64.0f)
-					.HeightOverride(64.0f)
+					SNew(SHorizontalBox)
+					+SHorizontalBox::Slot()
+					.AutoWidth()
 					.VAlign(VAlign_Center)
 					.HAlign(HAlign_Center)
+					.Padding(8.0f)
 					[
-						SNew(SImage)
-						.Image(SlateBrush)
+						SNew(SBox)
+						.WidthOverride(64.0f)
+						.HeightOverride(64.0f)
+						.VAlign(VAlign_Center)
+						.HAlign(HAlign_Center)
+						[
+							SNew(SImage)
+							.Image(SlateBrush)
+						]
 					]
-				]
-				+SHorizontalBox::Slot()
-				.FillWidth(1.0f)
-				.MaxWidth(600.0f)
-				.VAlign(VAlign_Center)
-				[
-					SNew(SVerticalBox)
-					+SVerticalBox::Slot()
-					.AutoHeight()
+					+SHorizontalBox::Slot()
+					.FillWidth(1.0f)
+					.VAlign(VAlign_Center)
 					[
-						SNew(STextBlock)
-						.Text(!Category.Title.IsEmpty() ? Category.Title : FText::FromString(CategoryName))
-						.TextStyle(&FEditorStyle::Get().GetWidgetStyle<FTextBlockStyle>("Tutorials.Browser.SummaryHeader"))
-						.HighlightText(HighlightText)
-						.HighlightColor(FEditorStyle::Get().GetColor("Tutorials.Browser.HighlightTextColor"))
-						.HighlightShape(FEditorStyle::Get().GetBrush("TextBlock.HighlightShape"))
+						SNew(SVerticalBox)
+						+SVerticalBox::Slot()
+						.AutoHeight()
+						[
+							SNew(STextBlock)
+							.Text(!Category.Title.IsEmpty() ? Category.Title : FText::FromString(CategoryName))
+							.TextStyle(&FEditorStyle::Get().GetWidgetStyle<FTextBlockStyle>("Tutorials.Browser.SummaryHeader"))
+							.HighlightText(HighlightText)
+							.HighlightColor(FEditorStyle::Get().GetColor("Tutorials.Browser.HighlightTextColor"))
+							.HighlightShape(FEditorStyle::Get().GetBrush("TextBlock.HighlightShape"))
+						]
+						+SVerticalBox::Slot()
+						.FillHeight(1.0f)
+						[
+							SNew(STextBlock)
+							.AutoWrapText(true)
+							.Text(Category.Description)
+							.TextStyle(&FEditorStyle::Get().GetWidgetStyle<FTextBlockStyle>("Tutorials.Browser.SummaryText"))
+							.HighlightText(HighlightText)
+							.HighlightColor(FEditorStyle::Get().GetColor("Tutorials.Browser.HighlightTextColor"))
+							.HighlightShape(FEditorStyle::Get().GetBrush("TextBlock.HighlightShape"))
+						]
 					]
-					+SVerticalBox::Slot()
-					.FillHeight(1.0f)
+					+SHorizontalBox::Slot()
+					.AutoWidth()
+					.HAlign(HAlign_Right)
+					.VAlign(VAlign_Center)
 					[
-						SNew(STextBlock)
-						.AutoWrapText(true)
-						.Text(Category.Description)
-						.TextStyle(&FEditorStyle::Get().GetWidgetStyle<FTextBlockStyle>("Tutorials.Browser.SummaryText"))
-						.HighlightText(HighlightText)
-						.HighlightColor(FEditorStyle::Get().GetColor("Tutorials.Browser.HighlightTextColor"))
-						.HighlightShape(FEditorStyle::Get().GetBrush("TextBlock.HighlightShape"))
+						SNew(SImage)
+						.Visibility(this, &FTutorialListEntry_Category::OnGetArrowVisibility)
+						.Image(FEditorStyle::Get().GetBrush("Tutorials.Browser.CategoryArrow"))
 					]
 				]
-				+SHorizontalBox::Slot()
-				.AutoWidth()
-				.HAlign(HAlign_Right)
-				.VAlign(VAlign_Center)
-				[
-					SNew(SImage)
-					.Visibility(this, &FTutorialListEntry_Category::OnGetArrowVisibility)
-					.Image(FEditorStyle::Get().GetBrush("Tutorials.Browser.CategoryArrow"))
-				]
 			]
 		];
 	}
@@ -227,6 +238,7 @@ public:
 		, OnTutorialSelected(InOnTutorialSelected)
 		, HighlightText(InHighlightText)
 		, SlateBrush(nullptr)
+		, LastUpdateTime(0.0f)
 	{
 		if(Tutorial->Texture != nullptr)
 		{
@@ -250,104 +262,106 @@ public:
 
 	virtual TSharedRef<ITableRow> OnGenerateTutorialRow(const TSharedRef<STableViewBase>& OwnerTable) const override
 	{
-		bHaveCompletedTutorial = GetDefault<UTutorialStateSettings>()->HaveCompletedTutorial(Tutorial);
-		bHaveSeenTutorial = false;
-		const int32 CurrentStage = GetDefault<UTutorialStateSettings>()->GetProgress(Tutorial, bHaveSeenTutorial);
-		Progress = (Tutorial->Stages.Num() > 0) ? (float)(CurrentStage + 1) / (float)Tutorial->Stages.Num() : 0.0f;
+		CacheProgress();
 
 		return SNew(STableRow<TSharedPtr<ITutorialListEntry>>, OwnerTable)
 		[
-			SAssignNew(LaunchButton, SButton)
-			.OnClicked(this, &FTutorialListEntry_Tutorial::OnClicked, false)
-			.ButtonStyle(&FEditorStyle::Get().GetWidgetStyle<FButtonStyle>("Tutorials.Browser.Button"))
-			.Content()
+			SNew(SBox)
+			.Padding(FMargin(0.0f, 2.0f))
 			[
-				SNew(SVerticalBox)
-				+SVerticalBox::Slot()
-				.AutoHeight()
+				SAssignNew(LaunchButton, SButton)
+				.OnClicked(this, &FTutorialListEntry_Tutorial::OnClicked, false)
+				.ButtonStyle(&FEditorStyle::Get().GetWidgetStyle<FButtonStyle>("Tutorials.Browser.Button"))
+				.ForegroundColor(FSlateColor::UseForeground())
+				.Content()
 				[
-					SNew(SHorizontalBox)
-					+SHorizontalBox::Slot()
-					.AutoWidth()
-					.VAlign(VAlign_Center)
-					.HAlign(HAlign_Center)
-					.Padding(8.0f)
+					SNew(SVerticalBox)
+					+SVerticalBox::Slot()
+					.AutoHeight()
 					[
-						SNew(SOverlay)
-						+SOverlay::Slot()
+						SNew(SHorizontalBox)
+						+SHorizontalBox::Slot()
+						.AutoWidth()
+						.VAlign(VAlign_Center)
+						.HAlign(HAlign_Center)
+						.Padding(8.0f)
 						[
-							SNew(SBox)
-							.WidthOverride(64.0f)
-							.HeightOverride(64.0f)
-							.VAlign(VAlign_Center)
-							.HAlign(HAlign_Center)
+							SNew(SOverlay)
+							+SOverlay::Slot()
+							[
+								SNew(SBox)
+								.WidthOverride(64.0f)
+								.HeightOverride(64.0f)
+								.VAlign(VAlign_Center)
+								.HAlign(HAlign_Center)
+								[
+									SNew(SImage)
+									.Image(SlateBrush)
+								]
+							]
+							+SOverlay::Slot()
+							.VAlign(VAlign_Bottom)
+							.HAlign(HAlign_Right)
 							[
 								SNew(SImage)
-								.Image(SlateBrush)
+								.ToolTipText(LOCTEXT("CompletedCheckToolTip", "This tutorial has been completed"))
+								.Visibility(this, &FTutorialListEntry_Tutorial::GetCompletedVisibility)
+								.Image(FEditorStyle::Get().GetBrush("Tutorials.Browser.Completed"))
 							]
 						]
-						+SOverlay::Slot()
-						.VAlign(VAlign_Bottom)
-						.HAlign(HAlign_Right)
+						+SHorizontalBox::Slot()
+						.FillWidth(1.0f)
+						.HAlign(HAlign_Fill)
+						.VAlign(VAlign_Center)
 						[
-							SNew(SImage)
-							.ToolTipText(LOCTEXT("CompletedCheckToolTip", "This tutorial has been completed"))
-							.Visibility(this, &FTutorialListEntry_Tutorial::GetCompletedVisibility)
-							.Image(FEditorStyle::Get().GetBrush("Tutorials.Browser.Completed"))
-						]
-					]
-					+SHorizontalBox::Slot()
-					.FillWidth(1.0f)
-					.HAlign(HAlign_Fill)
-					.VAlign(VAlign_Center)
-					[
-						SNew(SVerticalBox)
-						+SVerticalBox::Slot()
-						.AutoHeight()
-						[
-							SNew(SHorizontalBox)
-							+SHorizontalBox::Slot()
-							.FillWidth(1.0f)
+							SNew(SVerticalBox)
+							+SVerticalBox::Slot()
+							.AutoHeight()
 							[
-								SNew(STextBlock)
-								.Text(Tutorial->Title)
-								.TextStyle(&FEditorStyle::Get().GetWidgetStyle<FTextBlockStyle>("Tutorials.Browser.SummaryHeader"))
-								.HighlightText(HighlightText)
-								.HighlightColor(FEditorStyle::Get().GetColor("Tutorials.Browser.HighlightTextColor"))
-								.HighlightShape(FEditorStyle::Get().GetBrush("TextBlock.HighlightShape"))
+								SNew(SHorizontalBox)
+								+SHorizontalBox::Slot()
+								.FillWidth(1.0f)
+								[
+									SNew(STextBlock)
+									.Text(Tutorial->Title)
+									.TextStyle(&FEditorStyle::Get().GetWidgetStyle<FTextBlockStyle>("Tutorials.Browser.SummaryHeader"))
+									.HighlightText(HighlightText)
+									.HighlightColor(FEditorStyle::Get().GetColor("Tutorials.Browser.HighlightTextColor"))
+									.HighlightShape(FEditorStyle::Get().GetBrush("TextBlock.HighlightShape"))
+								]
+								+SHorizontalBox::Slot()
+								.AutoWidth()
+								.VAlign(VAlign_Center)
+								[
+									SNew(SButton)
+									.ToolTipText(LOCTEXT("RestartButtonToolTip", "Start this tutorial from the beginning"))
+									.Visibility(this, &FTutorialListEntry_Tutorial::GetRestartVisibility)
+									.OnClicked(this, &FTutorialListEntry_Tutorial::OnClicked, true)
+									.ButtonStyle(&FEditorStyle::Get().GetWidgetStyle<FButtonStyle>("Tutorials.Browser.Button"))
+									.Content()
+									[
+										SNew(SImage)
+										.Image(FEditorStyle::GetBrush("Tutorials.Browser.RestartButton"))
+									]
+								]
 							]
-							+SHorizontalBox::Slot()
-							.AutoWidth()
-							.VAlign(VAlign_Center)
+							+SVerticalBox::Slot()
+							.AutoHeight()
 							[
-								SNew(SButton)
-								.ToolTipText(LOCTEXT("RestartButtonToolTip", "Start this tutorial from the beginning"))
+								SNew(SBox)
 								.Visibility(this, &FTutorialListEntry_Tutorial::GetProgressVisibility)
-								.OnClicked(this, &FTutorialListEntry_Tutorial::OnClicked, true)
-								.ButtonStyle(&FEditorStyle::Get().GetWidgetStyle<FButtonStyle>("Tutorials.Browser.Button"))
-								.Content()
+								.HeightOverride(3.0f)
 								[
-									SNew(SImage)
-									.Image(FEditorStyle::GetBrush("Tutorials.Browser.RestartButton"))
+									SNew(SProgressBar)
+									.Percent(this, &FTutorialListEntry_Tutorial::GetProgress)
 								]
 							]
-						]
-						+SVerticalBox::Slot()
-						.AutoHeight()
-						[
-							SNew(SBox)
-							.Visibility(this, &FTutorialListEntry_Tutorial::GetProgressVisibility)
-							.HeightOverride(3.0f)
+							+SVerticalBox::Slot()
+							.FillHeight(1.0f)
 							[
-								SNew(SProgressBar)
-								.Percent(this, &FTutorialListEntry_Tutorial::GetProgress)
+								STutorialContent::GenerateContentWidget(Tutorial->SummaryContent, DocumentationPage, HighlightText)
 							]
 						]
-						+SVerticalBox::Slot()
-						.FillHeight(1.0f)
-						[
-							STutorialContent::GenerateContentWidget(Tutorial->SummaryContent, 500.0f, DocumentationPage, HighlightText)
-						]
 					]
 				]
 			]
@@ -386,6 +400,7 @@ public:
 
 	TOptional<float> GetProgress() const
 	{
+		CacheProgress();
 		return Progress;
 	}
 
@@ -393,17 +408,43 @@ public:
 	{
 		if(LaunchButton->IsHovered())
 		{
+			CacheProgress();
 			return LaunchButton->IsHovered() && bHaveSeenTutorial ? EVisibility::Visible : EVisibility::Hidden;
 		}
 
 		return EVisibility::Hidden;
 	}
 
+	EVisibility GetRestartVisibility() const
+	{
+		if(LaunchButton->IsHovered())
+		{
+			CacheProgress();
+			return LaunchButton->IsHovered() && bHaveSeenTutorial ? EVisibility::Visible : EVisibility::Collapsed;
+		}
+
+		return EVisibility::Collapsed;
+	}
+
 	EVisibility GetCompletedVisibility() const
 	{
+		CacheProgress();
 		return bHaveCompletedTutorial ? EVisibility::Visible : EVisibility::Hidden;
 	}
 
+	void CacheProgress() const
+	{
+		if(FPlatformTime::Seconds() - LastUpdateTime > TutorialBrowserConstants::ProgressUpdateInterval)
+		{
+			bHaveCompletedTutorial = GetDefault<UTutorialStateSettings>()->HaveCompletedTutorial(Tutorial);
+			bHaveSeenTutorial = false;
+			const int32 CurrentStage = GetDefault<UTutorialStateSettings>()->GetProgress(Tutorial, bHaveSeenTutorial);
+			Progress = (Tutorial->Stages.Num() > 0) ? (float)(CurrentStage + 1) / (float)Tutorial->Stages.Num() : 0.0f;
+
+			LastUpdateTime = FPlatformTime::Seconds();
+		}
+	}
+
 public:
 	/** Parent category */
 	TWeakPtr<ITutorialListEntry> ParentCategory;
@@ -437,104 +478,93 @@ public:
 
 	/** Cached tutorial progress */
 	mutable float Progress;
+
+	/** Last update time */
+	mutable float LastUpdateTime;
 };
 
 void STutorialsBrowser::Construct(const FArguments& InArgs)
 {
+	bNeedsRefresh = false;
+	RefreshTimer = TutorialBrowserConstants::RefreshTimerInterval;
+
 	OnClosed = InArgs._OnClosed;
 	OnLaunchTutorial = InArgs._OnLaunchTutorial;
 	ParentWindow = InArgs._ParentWindow;
 
+	FAssetRegistryModule& AssetRegistryModule = FModuleManager::LoadModuleChecked<FAssetRegistryModule>(TEXT("AssetRegistry"));
+	AssetRegistryModule.Get().OnAssetAdded().AddSP(this, &STutorialsBrowser::HandleAssetAdded);
+
 	ChildSlot
 	[
 		SNew(SVerticalBox)
 		+SVerticalBox::Slot()
-		.VAlign(VAlign_Center)
-		.HAlign(HAlign_Center)
+		.AutoHeight()
 		[
 			SNew(SBorder)
-			.Padding(18.0f)
-			.BorderImage(FEditorStyle::GetBrush("Tutorials.Border"))
+			.BorderImage(FEditorStyle::Get().GetBrush("ToolPanel.GroupBorder"))
+			.Padding(5.0f)
 			[
-				SNew(SVerticalBox)
-				+SVerticalBox::Slot()
-				.AutoHeight()
+				SNew(SHorizontalBox)
+				+SHorizontalBox::Slot()
+				.AutoWidth()
+				.VAlign(VAlign_Center)
 				[
-					// Title area
-					SNew(SHorizontalBox)
-					+SHorizontalBox::Slot()
-					.AutoWidth()
-					.VAlign(VAlign_Center)
-					[
-						SNew(SButton)
-						.OnClicked(this, &STutorialsBrowser::OnBackButtonClicked)
-						.IsEnabled(this, &STutorialsBrowser::IsBackButtonEnabled)
-						.ButtonStyle(&FEditorStyle::Get().GetWidgetStyle<FButtonStyle>("Tutorials.Content.Button"))
-						.Content()
-						[
-							SNew(SImage)
-							.Image(FEditorStyle::GetBrush("Tutorials.Browser.BackButton"))
-							.ColorAndOpacity(this, &STutorialsBrowser::GetBackButtonColor)
-						]
-					]
-					+SHorizontalBox::Slot()
-					.FillWidth(1.0f)
-					.Padding(10.0f)
+					SNew(SButton)
+					.OnClicked(this, &STutorialsBrowser::OnBackButtonClicked)
+					.IsEnabled(this, &STutorialsBrowser::IsBackButtonEnabled)
+					.ButtonStyle(&FEditorStyle::Get().GetWidgetStyle<FButtonStyle>("Tutorials.Browser.BackButton"))
+					.ForegroundColor(FSlateColor::UseForeground())
+					.Content()
 					[
-						SNew(STextBlock)
-						.Text(LOCTEXT("TutorialsTitle", "Unreal Editor Tutorials"))
-						.TextStyle( &FEditorStyle::Get().GetWidgetStyle<FTextBlockStyle>( "Tutorials.Browser.WelcomeHeader"))
-					]
-					+SHorizontalBox::Slot()
-					.AutoWidth()
-					.VAlign(VAlign_Top)
-					[
-						SNew(SButton)
-						.OnClicked(this, &STutorialsBrowser::OnCloseButtonClicked)
-						.ButtonStyle(&FEditorStyle::Get().GetWidgetStyle<FButtonStyle>("Tutorials.Browser.Button"))
-						.ContentPadding(0.0f)
-						[
-							SNew(SImage)
-							.Image(FEditorStyle::GetBrush("Symbols.X"))
-							.ColorAndOpacity(FLinearColor(0.0f, 0.0f, 0.0f, 1.0f))
-						]
+						SNew(SImage)
+						.Image(FEditorStyle::GetBrush("Tutorials.Browser.BackButton.Image"))
 					]
 				]
-				+SVerticalBox::Slot()
-				.AutoHeight()
-				[
-					SNew(SSearchBox)
-					.OnTextChanged(this, &STutorialsBrowser::OnSearchTextChanged)
-				]
-				+SVerticalBox::Slot()
-				.FillHeight(1.0f)
+				+SHorizontalBox::Slot()
+				.FillWidth(1.0f)
+				.VAlign(VAlign_Center)
+				.Padding(2.0f, 0.0f, 0.0f, 0.0f)
 				[
-					SNew(SBox)
-					.HeightOverride(600.0f)
-					.WidthOverride(600.0f)
+					SNew(SVerticalBox)
+					+SVerticalBox::Slot()
+					.AutoHeight()
+					.Padding(0.0f, 1.0f)
 					[
-						SNew(SScrollBox)
-						+SScrollBox::Slot()
-						[
-							SAssignNew(TutorialList, SListView<TSharedPtr<ITutorialListEntry>>)
-							.ItemHeight(128.0f)
-							.ListItemsSource(&FilteredEntries)
-							.OnGenerateRow(this, &STutorialsBrowser::OnGenerateTutorialRow)
-							.SelectionMode(ESelectionMode::None)
-						]
+						SAssignNew(BreadcrumbTrail, SBreadcrumbTrail<TSharedPtr<ITutorialListEntry>>)
+						.ButtonContentPadding(FMargin(1.0f, 1.0f))
+						.DelimiterImage(FEditorStyle::GetBrush("Tutorials.Browser.Breadcrumb"))
+						.TextStyle(FEditorStyle::Get(), "Tutorials.Browser.PathText")
+						.ShowLeadingDelimiter( true )
+						.InvertTextColorOnHover( false )
+						.OnCrumbClicked(this, &STutorialsBrowser::OnBreadcrumbClicked)
+					]
+					+SVerticalBox::Slot()
+					.AutoHeight()
+					.Padding(0.0f, 1.0f)
+					[
+						SNew(SSearchBox)
+						.OnTextChanged(this, &STutorialsBrowser::OnSearchTextChanged)
 					]
 				]
-				+SVerticalBox::Slot()
-				.AutoHeight()
-				.Padding(2.0f)
+			]
+		]
+		+SVerticalBox::Slot()
+		.FillHeight(1.0f)
+		.Padding(0.0f, 3.0f, 0.0f, 0.0f)
+		[
+			SNew(SVerticalBox)
+			+SVerticalBox::Slot()
+			.FillHeight(1.0f)
+			[
+				SNew(SScrollBox)
+				+SScrollBox::Slot()
 				[
-					SAssignNew(BreadcrumbTrail, SBreadcrumbTrail<TSharedPtr<ITutorialListEntry>>)
-					.ButtonContentPadding(FMargin(1.0f, 1.0f))
-					.DelimiterImage(FEditorStyle::GetBrush("Tutorials.Browser.Breadcrumb"))
-					.TextStyle(FEditorStyle::Get(), "Tutorials.Browser.PathText")
-					.ShowLeadingDelimiter( true )
-					.InvertTextColorOnHover( false )
-					.OnCrumbClicked(this, &STutorialsBrowser::OnBreadcrumbClicked)
+					SAssignNew(TutorialList, SListView<TSharedPtr<ITutorialListEntry>>)
+					.ItemHeight(128.0f)
+					.ListItemsSource(&FilteredEntries)
+					.OnGenerateRow(this, &STutorialsBrowser::OnGenerateTutorialRow)
+					.SelectionMode(ESelectionMode::None)
 				]
 			]
 		]
@@ -545,6 +575,18 @@ void STutorialsBrowser::Construct(const FArguments& InArgs)
 	RebuildCrumbs();
 }
 
+void STutorialsBrowser::Tick( const FGeometry& AllottedGeometry, const double InCurrentTime, const float InDeltaTime )
+{
+	RefreshTimer -= InDeltaTime;
+
+	if(bNeedsRefresh && RefreshTimer <= 0.0f)
+	{
+		ReloadTutorials();
+		bNeedsRefresh = false;
+		RefreshTimer = TutorialBrowserConstants::RefreshTimerInterval;
+	}
+}
+
 void STutorialsBrowser::SetFilter(const FString& InFilter)
 {
 	CategoryFilter = InFilter;
@@ -766,11 +808,6 @@ bool STutorialsBrowser::IsBackButtonEnabled() const
 	return false;
 }
 
-FSlateColor STutorialsBrowser::GetBackButtonColor() const
-{
-	return IsBackButtonEnabled() ? FLinearColor(0.0f, 0.0f, 0.0f, 1.0f) : FLinearColor(1.0f, 1.0f, 1.0f, 0.25f);
-}
-
 void STutorialsBrowser::OnTutorialSelected(UEditorTutorial* InTutorial, bool bRestart)
 {
 	if( FEngineAnalytics::IsAvailable() && InTutorial != nullptr )
@@ -961,4 +998,20 @@ void STutorialsBrowser::RebuildCrumbs()
 	}
 }
 
+void STutorialsBrowser::HandleAssetAdded(const FAssetData& InAssetData)
+{
+	if(InAssetData.AssetClass == UBlueprint::StaticClass()->GetFName())
+	{
+		const FString* ParentClassPath = InAssetData.TagsAndValues.Find(TEXT("ParentClass"));
+		if(ParentClassPath != nullptr)
+		{
+			UClass* ParentClass = FindObject<UClass>(NULL, **ParentClassPath);
+			if(ParentClass == UEditorTutorial::StaticClass())
+			{
+				bNeedsRefresh = true;
+			}
+		}
+	}
+}
+
 #undef LOCTEXT_NAMESPACE
\ No newline at end of file
diff --git a/Engine/Source/Editor/IntroTutorials/Private/STutorialsBrowser.h b/Engine/Source/Editor/IntroTutorials/Private/STutorialsBrowser.h
--- a/Engine/Source/Editor/IntroTutorials/Private/STutorialsBrowser.h
+++ b/Engine/Source/Editor/IntroTutorials/Private/STutorialsBrowser.h
@@ -51,6 +51,9 @@ class STutorialsBrowser : public SCompoundWidget
 	/** Reload all tutorials that we know about */
 	void ReloadTutorials();
 
+	/** SWidget implementation */
+	virtual void Tick( const FGeometry& AllottedGeometry, const double InCurrentTime, const float InDeltaTime ) override;
+
 protected:
 	/** Handle generating a table row in the browser */
 	TSharedRef<ITableRow> OnGenerateTutorialRow(TSharedPtr<ITutorialListEntry> InItem, const TSharedRef<STableViewBase>& OwnerTable) const;
@@ -64,9 +67,6 @@ protected:
 	/** Handle whether the back button can be clicked */
 	bool IsBackButtonEnabled() const;
 
-	/** Handle back button color - used to dim the button when disabled */
-	FSlateColor GetBackButtonColor() const;
-
 	/** Delegate handler fired when a tutorial is selected form the browser */
 	void OnTutorialSelected(UEditorTutorial* InTutorial, bool bRestart);
 
@@ -97,6 +97,9 @@ protected:
 	/** Rebuild the breadcrumb trail according to the current category */
 	void RebuildCrumbs();
 
+	/** Handle an asset being added - rebuild our list if required */
+	void HandleAssetAdded(const FAssetData& InAssetData);
+
 private:
 
 	/** Root entry of the tutorials tree */
@@ -131,4 +134,10 @@ private:
 
 	/** Breadcrumb trail for path display */
 	TSharedPtr<SBreadcrumbTrail<TSharedPtr<ITutorialListEntry>>> BreadcrumbTrail;
+
+	/** Whether we need to refresh the content in the browser */
+	bool bNeedsRefresh;
+
+	/** Prevent us from refreshing too often */
+	float RefreshTimer;
 };
\ No newline at end of file
