Tutorial Blueprint fixes

Tutorials now return a valid UWorld, enabling audio playback, timers and all sorts of good stuff.
Stage started event now fires on first stage.
Skipping a stage from BP now rebuilds content properly.
Added events for tutorial launch and close.

diff --git a/Engine/Source/Editor/IntroTutorials/Private/EditorTutorial.cpp b/Engine/Source/Editor/IntroTutorials/Private/EditorTutorial.cpp
--- a/Engine/Source/Editor/IntroTutorials/Private/EditorTutorial.cpp
+++ b/Engine/Source/Editor/IntroTutorials/Private/EditorTutorial.cpp
@@ -8,6 +8,16 @@ UEditorTutorial::UEditorTutorial(const FObjectInitializer& ObjectInitializer)
 {
 }
 
+UWorld* UEditorTutorial::GetWorld() const
+{
+#if WITH_EDITOR
+	if (GIsEditor)
+	{
+		return GWorld;
+	}
+#endif // WITH_EDITOR
+	return GEngine->GetWorldContexts()[0].World();
+}
 
 void UEditorTutorial::GoToNextTutorialStage()
 {
@@ -44,6 +54,20 @@ void UEditorTutorial::HandleTutorialStageEnded(FName StageName)
 }
 
 
+void UEditorTutorial::HandleTutorialLaunched()
+{
+	FEditorScriptExecutionGuard ScriptGuard;
+	OnTutorialLaunched();
+}
+
+
+void UEditorTutorial::HandleTutorialClosed()
+{
+	FEditorScriptExecutionGuard ScriptGuard;
+	OnTutorialClosed();
+}
+
+
 void UEditorTutorial::OpenAsset(UObject* Asset)
 {
 	FAssetEditorManager::Get().OpenEditorForAsset(Asset);
diff --git a/Engine/Source/Editor/IntroTutorials/Private/STutorialRoot.cpp b/Engine/Source/Editor/IntroTutorials/Private/STutorialRoot.cpp
--- a/Engine/Source/Editor/IntroTutorials/Private/STutorialRoot.cpp
+++ b/Engine/Source/Editor/IntroTutorials/Private/STutorialRoot.cpp
@@ -144,6 +144,16 @@ void STutorialRoot::LaunchTutorial(UEditorTutorial* InTutorial, bool bInRestart,
 				TutorialWidget.Value.Pin()->LaunchTutorial(bIsNavigationWindow, InOnTutorialClosed, InOnTutorialExited);
 			}
 		}
+
+		if (CurrentTutorial != nullptr)
+		{
+			CurrentTutorial->HandleTutorialLaunched();
+		}
+
+		if (CurrentTutorial != nullptr && CurrentTutorialStage < CurrentTutorial->Stages.Num())
+		{
+			CurrentTutorial->HandleTutorialStageStarted(CurrentTutorial->Stages[CurrentTutorialStage].Name);
+		}
 	}
 }
 
@@ -161,15 +171,6 @@ void STutorialRoot::CloseAllTutorialContent()
 void STutorialRoot::HandleNextClicked(TWeakPtr<SWindow> InNavigationWindow)
 {
 	GoToNextStage(InNavigationWindow);
-
-	for(auto& TutorialWidget : TutorialWidgets)
-	{
-		if(TutorialWidget.Value.IsValid())
-		{
-			TSharedPtr<SEditorTutorials> PinnedTutorialWidget = TutorialWidget.Value.Pin();
-			PinnedTutorialWidget->RebuildCurrentContent();
-		}
-	}
 }
 
 void STutorialRoot::HandleBackClicked()
@@ -199,6 +200,7 @@ void STutorialRoot::HandleHomeClicked()
 {
 	if(CurrentTutorial != nullptr)
 	{
+		CurrentTutorial->HandleTutorialClosed();
 		GetMutableDefault<UTutorialStateSettings>()->RecordProgress(CurrentTutorial, CurrentTutorialStage);
 		GetMutableDefault<UTutorialStateSettings>()->SaveProgress();
 	}
@@ -287,10 +289,24 @@ void STutorialRoot::GoToNextStage(TWeakPtr<SWindow> InNavigationWindow)
 			CurrentTutorial->HandleTutorialStageStarted(CurrentTutorial->Stages[CurrentTutorialStage].Name);
 		}
 	}
+
+	for(auto& TutorialWidget : TutorialWidgets)
+	{
+		if(TutorialWidget.Value.IsValid())
+		{
+			TSharedPtr<SEditorTutorials> PinnedTutorialWidget = TutorialWidget.Value.Pin();
+			PinnedTutorialWidget->RebuildCurrentContent();
+		}
+	}
 }
 
 void STutorialRoot::HandleCloseClicked()
 {
+	if(CurrentTutorial != nullptr)
+	{
+		CurrentTutorial->HandleTutorialClosed();
+	}
+
 	// submit analytics data
 	if( FEngineAnalytics::IsAvailable() && CurrentTutorial != nullptr && CurrentTutorialStage < CurrentTutorial->Stages.Num() )
 	{
diff --git a/Engine/Source/Editor/IntroTutorials/Public/EditorTutorial.h b/Engine/Source/Editor/IntroTutorials/Public/EditorTutorial.h
--- a/Engine/Source/Editor/IntroTutorials/Public/EditorTutorial.h
+++ b/Engine/Source/Editor/IntroTutorials/Public/EditorTutorial.h
@@ -206,7 +206,7 @@ struct INTROTUTORIALS_API FTutorialStage
 };
 
 /** An asset used to build a stage-by-stage tutorial in the editor */
-UCLASS(Blueprintable, EditInlineNew)
+UCLASS(Blueprintable, hideCategories=(Object), EditInlineNew)
 class INTROTUTORIALS_API UEditorTutorial : public UObject
 {
 	GENERATED_UCLASS_BODY()
@@ -255,6 +255,9 @@ class INTROTUTORIALS_API UEditorTutorial : public UObject
 	UPROPERTY(EditAnywhere, Category="Tutorial")
 	bool bHideInBrowser;
 
+	/** UObject implementation */
+	virtual UWorld* GetWorld() const override;
+
 public:
 	/** Called when a tutorial stage is started */
 	void HandleTutorialStageStarted(FName StageName);
@@ -265,6 +268,12 @@ public:
 	/** Called each tick so the Blueprint can optionally complete or skip stages */
 	void HandleTickCurrentStage(FName StageName);
 
+	/** Called when a tutorial is launched */
+	void HandleTutorialLaunched();
+
+	/** Called when a tutorial is closed */
+	void HandleTutorialClosed();
+
 protected:
 	/** Event fired when a tutorial stage begins */
 	UFUNCTION(BlueprintImplementableEvent, Category="Tutorial")
@@ -274,6 +283,14 @@ protected:
 	UFUNCTION(BlueprintImplementableEvent, Category="Tutorial")
 	void OnTutorialStageEnded(FName StageName) const;
 
+	/** Event fired when a tutorial is launched */
+	UFUNCTION(BlueprintImplementableEvent, Category = "Tutorial")
+	void OnTutorialLaunched() const;
+
+	/** Event fired when a tutorial is closed */
+	UFUNCTION(BlueprintImplementableEvent, Category = "Tutorial")
+	void OnTutorialClosed() const;
+
 	/** Advance to the next stage of a tutorial */
 	UFUNCTION(BlueprintCallable, Category="Tutorial")
 	static void GoToNextTutorialStage();
