Added Tab focus field for tutorials.

diff --git a/Engine/Source/Editor/IntroTutorials/Private/IntroTutorials.cpp b/Engine/Source/Editor/IntroTutorials/Private/IntroTutorials.cpp
--- a/Engine/Source/Editor/IntroTutorials/Private/IntroTutorials.cpp
+++ b/Engine/Source/Editor/IntroTutorials/Private/IntroTutorials.cpp
@@ -373,6 +373,11 @@ FOnIsPicking& FIntroTutorials::OnIsPicking()
 	return OnIsPickingDelegate;
 }
 
+FOnValidatePickingCandidate& FIntroTutorials::OnValidatePickingCandidate()
+{
+	return OnValidatePickingCandidateDelegate;
+}
+
 void FIntroTutorials::SummonTutorialBrowser()
 {
 	if(TutorialRoot.IsValid())
diff --git a/Engine/Source/Editor/IntroTutorials/Private/IntroTutorials.h b/Engine/Source/Editor/IntroTutorials/Private/IntroTutorials.h
--- a/Engine/Source/Editor/IntroTutorials/Private/IntroTutorials.h
+++ b/Engine/Source/Editor/IntroTutorials/Private/IntroTutorials.h
@@ -5,6 +5,9 @@
 /** Query delegate to see if we are in picking mode */
 DECLARE_DELEGATE_RetVal_OneParam(bool, FOnIsPicking, FName& /* OutWidgetNameToHighlight */);
 
+/** Query delegate for name of any picked widget and check if it matches the 'pickable' name of the given widget */
+DECLARE_DELEGATE_RetVal_ThreeParams(bool, FOnValidatePickingCandidate, TSharedRef<SWidget> /*InWidget*/, FName& /* OutWidgetNameToHighlight */, bool& /*bOutShouldHighlight*/);
+
 class UEditorTutorial;
 struct FTutorialContent;
 class STutorialRoot;
@@ -17,6 +20,9 @@ public:
 
 	/** Get the delegate used to check for whether we are picking widgets */
 	FOnIsPicking& OnIsPicking();
+	
+	/** Get the delegate used to validate the given widget for pickings */
+	FOnValidatePickingCandidate& OnValidatePickingCandidate();
 
 	void GoToPreviousStage();
 
@@ -112,6 +118,9 @@ private:
 
 	/** Delegate used to determine whether we are in picking mode */
 	FOnIsPicking OnIsPickingDelegate;
+	
+	/** Delegate used to determine if we are in picking mode, get the name of any picked widget and check if it matches the 'pickable' name of the given widget */
+	FOnValidatePickingCandidate OnValidatePickingCandidateDelegate;
 
 	/** Root widget for tutorial overlay system */
 	TSharedPtr<STutorialRoot> TutorialRoot;
diff --git a/Engine/Source/Editor/IntroTutorials/Private/STutorialOverlay.cpp b/Engine/Source/Editor/IntroTutorials/Private/STutorialOverlay.cpp
--- a/Engine/Source/Editor/IntroTutorials/Private/STutorialOverlay.cpp
+++ b/Engine/Source/Editor/IntroTutorials/Private/STutorialOverlay.cpp
@@ -9,6 +9,7 @@
 #include "LevelEditor.h"
 #include "BlueprintEditorUtils.h"
 #include "Guid.h"
+#include "BlueprintEditor.h"
 
 static FName IntroTutorialsModuleName("IntroTutorials");
 
@@ -129,29 +130,35 @@ int32 STutorialOverlay::OnPaint( const FPaintArgs& Args, const FGeometry& Allott
 
 int32 STutorialOverlay::TraverseWidgets(TSharedRef<SWidget> InWidget, const FGeometry& InGeometry, const FSlateRect& MyClippingRect, FSlateWindowElementList& OutDrawElements, int32 LayerId) const
 {
+	bool bIsPicking = false;
+	bool bShouldHighlight = false;
+	bool bShouldDraw = false;
+	FName WidgetNameToHighlight = NAME_None;
+	FIntroTutorials& IntroTutorials = FModuleManager::Get().GetModuleChecked<FIntroTutorials>(IntroTutorialsModuleName);
+	if (IntroTutorials.OnValidatePickingCandidate().IsBound())
+	{
+		bIsPicking = IntroTutorials.OnValidatePickingCandidate().Execute(InWidget,WidgetNameToHighlight,bShouldHighlight);
+	}
+
+	// First draw the widget if we should
 	TSharedPtr<FTagMetaData> MetaData = InWidget->GetMetaData<FTagMetaData>();
-	const FName Tag = (MetaData.IsValid() && MetaData->Tag.IsValid()) ? MetaData->Tag : InWidget->GetTag();	
-	if(Tag != NAME_None || MetaData.IsValid())
+	const FName Tag = (MetaData.IsValid() && MetaData->Tag.IsValid()) ? MetaData->Tag : InWidget->GetTag();
+	if (Tag != NAME_None || MetaData.IsValid())
 	{
 		// we are a named widget - ask it to draw
 		OnPaintNamedWidget.Broadcast(InWidget, InGeometry);
 		OnWidgetWasDrawn.ExecuteIfBound(Tag);
+	}
 
-		// if we are picking, we need to draw an outline here
-		FName WidgetNameToHighlight = NAME_None;
-		bool bIsPicking = false;
-		FIntroTutorials& IntroTutorials = FModuleManager::Get().GetModuleChecked<FIntroTutorials>(IntroTutorialsModuleName);
-		if(IntroTutorials.OnIsPicking().IsBound())
-		{
-			bIsPicking = IntroTutorials.OnIsPicking().Execute(WidgetNameToHighlight);
-		}
-	
-		if(WidgetNameToHighlight != NAME_None || bIsPicking)
+	// Next check and draw the highlight as appropriate
+	if (bIsPicking == true)
+	{
+		// if we are picking, we need to draw an outline here	
+		if(WidgetNameToHighlight != NAME_None )
 		{
-			const bool bHighlight = WidgetNameToHighlight != NAME_None && WidgetNameToHighlight == Tag;
-			if(bIsPicking || (!bIsPicking && bHighlight))
+			if(bIsPicking == true)
 			{
-				const FLinearColor Color = bIsPicking && bHighlight ? FLinearColor::Green : FLinearColor::White;
+				const FLinearColor Color = bIsPicking && bShouldHighlight ? FLinearColor::Green : FLinearColor::White;
 				FSlateDrawElement::MakeBox(OutDrawElements, LayerId++, InGeometry.ToPaintGeometry(), FCoreStyle::Get().GetBrush(TEXT("Debug.Border")), MyClippingRect, ESlateDrawEffect::None, Color);
 			}
 		}	
@@ -178,20 +185,64 @@ void STutorialOverlay::PerformWidgetInteractions(const FTutorialWidgetContent &W
 
 void STutorialOverlay::OpenBrowserForWidgetAnchor(const FTutorialWidgetContent &WidgetContent)
 {
-	bool bTabOpened = false;
-	// Open the required tab if we found it in the map
-	if (WidgetContent.WidgetAnchor.TabTypeToOpen.IsEmpty() == false)
+	//bool bTabOpened = false;
+	TSharedPtr<FTabManager> TabManager;
+	TArray<FString> AssetPaths;
+	FString Name;
+
+	// Try looking for the object via the ID first
+	Name = WidgetContent.WidgetAnchor.WrapperIdentifier.ToString();
+	AssetPaths.Add(Name);
+	// Opening the editor will force the object to be loaded if it exists
+	FAssetEditorManager::Get().OpenEditorsForAssets(AssetPaths);
+	UObject* AssetObject = FindObject<UObject>(ANY_PACKAGE, *Name);
+	IAssetEditorInstance* AssetEditor = FAssetEditorManager::Get().FindEditorForAsset(AssetObject, false);
+	// If we now have a valid asset editor get its tabmanager - we will use this to open/focus a tab if there is one
+	if (AssetEditor != nullptr)
+	{
+		FAssetEditorToolkit* ToolkitEditor = static_cast<FAssetEditorToolkit*>(FAssetEditorManager::Get().FindEditorForAsset(AssetObject, false));
+		if ((ToolkitEditor != nullptr) && (WidgetContent.WidgetAnchor.TabToFocusOrOpen.IsEmpty() == false))
+		{
+			TabManager = ToolkitEditor->GetTabManager();
+		}
+	}
+
+	// If we don't have a valid tab manager we should now check to see if we can find a blueprint relevant to this node and open the editor for that (Then try to get the tabmanager from that)
+	if (TabManager.IsValid() == false)
 	{
-		FLevelEditorModule& LevelEditorModule = FModuleManager::GetModuleChecked<FLevelEditorModule>(TEXT("LevelEditor"));
-		TSharedPtr<FTabManager> LevelEditorTabManager = LevelEditorModule.GetLevelEditorTabManager();
-		TSharedPtr<SDockTab> DocTab = LevelEditorTabManager->InvokeTab(FName(*WidgetContent.WidgetAnchor.TabTypeToOpen));
-		bTabOpened = DocTab.IsValid();
+		// Remove the prefix from the name
+		AssetPaths.Empty();
+		int32 Space = WidgetContent.WidgetAnchor.OuterName.Find(TEXT(" "));
+		Name = WidgetContent.WidgetAnchor.OuterName.RightChop(Space + 1);
+		AssetPaths.Add(Name);
+		FAssetEditorManager::Get().OpenEditorsForAssets(AssetPaths);
+		UObject* Blueprint = FindObject<UObject>(ANY_PACKAGE, *Name);
+
+		// If we found a blueprint
+		if (Blueprint != nullptr)
+		{
+			IBlueprintEditor* BlueprintEditor = (FBlueprintEditor*)FAssetEditorManager::Get().FindEditorForAsset(Blueprint, false);
+			if ((BlueprintEditor != nullptr) && (WidgetContent.WidgetAnchor.TabToFocusOrOpen.IsEmpty() == false))
+			{
+				TabManager = BlueprintEditor->GetTabManager();
+			}
+		}
 	}
-	if (bTabOpened == false)
+
+	// Invoke any tab
+	if (WidgetContent.WidgetAnchor.TabToFocusOrOpen.IsEmpty() == false)
 	{
-		//bTabOpened = FAssetEditorManager::Get().OpenEditorForAsset()
+		if (TabManager.IsValid() == true)
+		{
+			TSharedRef<SDockTab> IT = TabManager->InvokeTab(FTabId(*WidgetContent.WidgetAnchor.TabToFocusOrOpen));
+		}
+		else
+		{
+			FLevelEditorModule& LevelEditorModule = FModuleManager::GetModuleChecked<FLevelEditorModule>(TEXT("LevelEditor"));
+			TSharedPtr<FTabManager> LevelEditorTabManager = LevelEditorModule.GetLevelEditorTabManager();
+			LevelEditorTabManager->InvokeTab(FName(*WidgetContent.WidgetAnchor.TabToFocusOrOpen));
+		}
 	}
-	
 }
 
 void STutorialOverlay::FocusOnAnyBlueprintNodes(const FTutorialWidgetContent &WidgetContent)
diff --git a/Engine/Source/Editor/IntroTutorials/Private/TutorialStructCustomization.cpp b/Engine/Source/Editor/IntroTutorials/Private/TutorialStructCustomization.cpp
--- a/Engine/Source/Editor/IntroTutorials/Private/TutorialStructCustomization.cpp
+++ b/Engine/Source/Editor/IntroTutorials/Private/TutorialStructCustomization.cpp
@@ -6,6 +6,7 @@
 #include "STutorialEditableText.h"
 #include "ISlateMetaData.h"
 #include "TutorialMetaData.h"
+#include "SDockTab.h"
 
 #define LOCTEXT_NAMESPACE "TutorialStructCustomization"
 
@@ -176,58 +177,93 @@ void FTutorialContentCustomization::CustomizeChildren( TSharedRef<class IPropert
 /** 'Tooltip' window to indicate what is currently being picked and how to abort the picker (Esc) */
 class SWidgetPickerFloatingWindow : public SCompoundWidget
 {
-	SLATE_BEGIN_ARGS( SWidgetPickerFloatingWindow ) {}
+	SLATE_BEGIN_ARGS(SWidgetPickerFloatingWindow){}
 
 	SLATE_ARGUMENT(TWeakPtr<SWindow>, ParentWindow)
+	SLATE_ARGUMENT(TSharedPtr<class IPropertyHandle>, FriendlyNameProperty)
+	SLATE_ARGUMENT(FName, SpecificWidgetType)
 
 	SLATE_END_ARGS()
 
-	void Construct(const FArguments& InArgs, TSharedRef<class IPropertyHandle> InStructPropertyHandle)
+	void Construct(const FArguments& InArgs, TSharedRef<class IPropertyHandle> InStructPropertyHandle, TSharedRef<class IPropertyHandle> InPickPropertyHandle)
 	{
 		StructPropertyHandle = InStructPropertyHandle;
-		ParentWindow = InArgs._ParentWindow;
+		PickPropertyHandle = InPickPropertyHandle;
 
+		ParentWindow = InArgs._ParentWindow;
+		SpecificWidgetType = InArgs._SpecificWidgetType;
+		FriendlyNameProperty = InArgs._FriendlyNameProperty;
+		
 		ChildSlot
 		[
 			SNew(SToolTip)
 			.Text(this, &SWidgetPickerFloatingWindow::GetPickerStatusText)
 		];
 	}
-
+	
+	/* Returns the name of the picked widget */
 	FName GetPickedWidgetName()
 	{
 		return PickedWidgetName;
 	}
 
+	/** 
+	 * Return the name of the given widget (Will filter out widgets that do no match the specific type if applicable
+	 *
+	 * @param InWidget	The widget to get the pickable name of
+	 * @returns Pickable Name of the widget (or None if it doesnt match a specific widget type)
+	 *
+	 */
+	FName GetPickableNameForWidget(TSharedRef<SWidget> InWidget) const
+	{
+		FName PickableName;
+		if (InWidget->GetTag() != NAME_None)
+		{
+			PickableName = InWidget->GetTag();
+		}
+		else
+		{
+			// If we have specified a specific widget to to pick check this one matches
+			TSharedPtr<FTagMetaData> MetaData = InWidget->GetMetaData<FTagMetaData>();
+			if (MetaData.IsValid())
+			{
+				PickableName = MetaData->Tag;
+			}
+			else if (SpecificWidgetType != NAME_None)
+			{
+				FName TheType = InWidget->GetType();
+				if (TheType == SpecificWidgetType)
+				{
+					if (SpecificWidgetType == FName("SDockTab"))
+					{
+						TSharedRef<SDockTab> DockTab = StaticCastSharedRef<SDockTab>(InWidget);
+						FString TabIdent = DockTab.Get().GetLayoutIdentifier().ToString();
+						PickableName = FName(*TabIdent);
+					}
+				}
+			}
+		}
+		return PickableName;
+	}
+	
 private:
 	virtual void Tick( const FGeometry& AllottedGeometry, const double InCurrentTime, const float InDeltaTime ) override
 	{
 		PickedWidgetName = NAME_None;
 		PickedAllMetaData.Reset();
-
+		
 		FWidgetPath Path = FSlateApplication::Get().LocateWindowUnderMouse(FSlateApplication::Get().GetCursorPos(), FSlateApplication::Get().GetInteractiveTopLevelWindows(), true);
 
 		for(int32 PathIndex = Path.Widgets.Num() - 1; PathIndex >= 0; PathIndex--)
 		{
 			TSharedRef<SWidget> PathWidget = Path.Widgets[PathIndex].Widget;
-			if(PathWidget->GetTag() != NAME_None)
+						
+			PickedWidgetName = GetPickableNameForWidget(PathWidget);
+			if (PickedWidgetName != NAME_None)
 			{
-				PickedWidgetName = PathWidget->GetTag();
+				PickedAllMetaData = PathWidget->GetAllMetaData<FTagMetaData>();
 				break;
-			}
-			else
-			{
-				TSharedPtr<FTagMetaData> MetaData = PathWidget->GetMetaData<FTagMetaData>();
-				if(MetaData.IsValid())
-				{
-					PickedWidgetName = MetaData->Tag;
-					if (PickedWidgetName != NAME_None)
-					{
-						PickedAllMetaData = PathWidget->GetAllMetaData<FTagMetaData>();
-					}
-					break;
-				}
-			}		
+			}			
 		}
 
 		// kind of a hack, but we need to maintain keyboard focus otherwise we wont get our keypress to 'pick'
@@ -238,7 +274,7 @@ private:
 			ParentWindow.Pin()->MoveWindowTo(FSlateApplication::Get().GetCursorPos() + FSlateApplication::Get().GetCursorSize());
 		}
 	}
-
+		
 	FText GetPickerStatusText() const
 	{
 		return FText::Format(LOCTEXT("TootipHint", "{0} (Esc to pick)"), FText::FromName(PickedWidgetName));
@@ -248,42 +284,66 @@ private:
 	{
 		if(InKeyboardEvent.GetKey() == EKeys::Escape)
 		{
-			TSharedPtr<IPropertyHandle> WrapperIdentifierProperty = StructPropertyHandle->GetChildHandle(GET_MEMBER_NAME_CHECKED(FTutorialContentAnchor, WrapperIdentifier));
-			WrapperIdentifierProperty->SetValue(PickedWidgetName);
-
-			TSharedPtr<IPropertyHandle> TypeProperty = StructPropertyHandle->GetChildHandle(GET_MEMBER_NAME_CHECKED(FTutorialContentAnchor, Type));
-			TypeProperty->SetValue((uint8)ETutorialAnchorIdentifier::NamedWidget);
-			
-			// Set the friendly name to the PickedWidget name - we might not have any metadata
-			StructPropertyHandle->GetChildHandle(GET_MEMBER_NAME_CHECKED(FTutorialContentAnchor, FriendlyName))->SetValue(PickedWidgetName);
-			// Reset the other fields
-			StructPropertyHandle->GetChildHandle(GET_MEMBER_NAME_CHECKED(FTutorialContentAnchor, GUIDString))->SetValue(FString());
-			StructPropertyHandle->GetChildHandle(GET_MEMBER_NAME_CHECKED(FTutorialContentAnchor, OuterName))->SetValue(FString());
-			StructPropertyHandle->GetChildHandle(GET_MEMBER_NAME_CHECKED(FTutorialContentAnchor, TabTypeToOpen))->SetValue(FString());
-
-			for (const auto& MetaDataEntry : PickedAllMetaData)
+			if( InKeyboardEvent.IsLeftControlDown() == false )
 			{
-				if (MetaDataEntry->IsOfType<FTutorialMetaData>())
-				{
-					TSharedRef<FTutorialMetaData> TutorialMeta = StaticCastSharedRef<FTutorialMetaData>(MetaDataEntry);
-					StructPropertyHandle->GetChildHandle(GET_MEMBER_NAME_CHECKED(FTutorialContentAnchor, FriendlyName))->SetValue(TutorialMeta->FriendlyName);
-					StructPropertyHandle->GetChildHandle(GET_MEMBER_NAME_CHECKED(FTutorialContentAnchor, TabTypeToOpen))->SetValue(TutorialMeta->TabTypeToOpen);
-				}
-				else if (MetaDataEntry->IsOfType<FGraphNodeMetaData>())
+				// We cant set a parameter if this isn't valid !
+				check(PickPropertyHandle.IsValid());
+				PickPropertyHandle->SetValue(PickedWidgetName);
+
+				TSharedPtr<IPropertyHandle> TypeProperty = StructPropertyHandle->GetChildHandle(GET_MEMBER_NAME_CHECKED(FTutorialContentAnchor, Type));
+				TypeProperty->SetValue((uint8)ETutorialAnchorIdentifier::NamedWidget);
+
+				FString FriendlyNameToSet = PickedWidgetName.ToString();
+
+				// Reset the other fields
+				StructPropertyHandle->GetChildHandle(GET_MEMBER_NAME_CHECKED(FTutorialContentAnchor, GUIDString))->SetValue(FString());
+				StructPropertyHandle->GetChildHandle(GET_MEMBER_NAME_CHECKED(FTutorialContentAnchor, OuterName))->SetValue(FString());
+				StructPropertyHandle->GetChildHandle(GET_MEMBER_NAME_CHECKED(FTutorialContentAnchor, TabToFocusOrOpen))->SetValue(FString());
+
+				// Handle custom widget type picks
+				if ((SpecificWidgetType.IsValid() == true) && (SpecificWidgetType != NAME_None))
 				{
-					TSharedRef<FGraphNodeMetaData> GraphNodeMeta = StaticCastSharedRef<FGraphNodeMetaData>(MetaDataEntry);
-					StructPropertyHandle->GetChildHandle(GET_MEMBER_NAME_CHECKED(FTutorialContentAnchor, GUIDString))->SetValue(GraphNodeMeta->GUID.ToString());
-					StructPropertyHandle->GetChildHandle(GET_MEMBER_NAME_CHECKED(FTutorialContentAnchor, OuterName))->SetValue(GraphNodeMeta->OuterName);
-					StructPropertyHandle->GetChildHandle(GET_MEMBER_NAME_CHECKED(FTutorialContentAnchor, FriendlyName))->SetValue(GraphNodeMeta->FriendlyName);
-					StructPropertyHandle->GetChildHandle(GET_MEMBER_NAME_CHECKED(FTutorialContentAnchor, TabTypeToOpen))->SetValue(GraphNodeMeta->TabTypeToOpen);
+					if (SpecificWidgetType == FName("SDockTab"))
+					{
+						StructPropertyHandle->GetChildHandle(GET_MEMBER_NAME_CHECKED(FTutorialContentAnchor, TabToFocusOrOpen))->SetValue(PickedWidgetName);
+					}
 				}
-				else
+
+				for (const auto& MetaDataEntry : PickedAllMetaData)
 				{
-					TSharedRef<FTagMetaData> GraphNodeMeta = StaticCastSharedRef<FTagMetaData>(MetaDataEntry);
-					StructPropertyHandle->GetChildHandle(GET_MEMBER_NAME_CHECKED(FTutorialContentAnchor, FriendlyName))->SetValue(GraphNodeMeta->Tag);
+					if (MetaDataEntry->IsOfType<FGraphNodeMetaData>())
+					{
+						TSharedRef<FGraphNodeMetaData> GraphNodeMeta = StaticCastSharedRef<FGraphNodeMetaData>(MetaDataEntry);
+						StructPropertyHandle->GetChildHandle(GET_MEMBER_NAME_CHECKED(FTutorialContentAnchor, GUIDString))->SetValue(GraphNodeMeta->GUID.ToString());
+						StructPropertyHandle->GetChildHandle(GET_MEMBER_NAME_CHECKED(FTutorialContentAnchor, OuterName))->SetValue(GraphNodeMeta->OuterName);
+						StructPropertyHandle->GetChildHandle(GET_MEMBER_NAME_CHECKED(FTutorialContentAnchor, FriendlyName))->SetValue(GraphNodeMeta->FriendlyName);
+					}
+					else if (MetaDataEntry->IsOfType<FTutorialMetaData>())
+					{
+						TSharedRef<FTutorialMetaData> TutorialMeta = StaticCastSharedRef<FTutorialMetaData>(MetaDataEntry);
+						FriendlyNameToSet = TutorialMeta->FriendlyName;
+
+						// TabTypeToOpen only really applies to specifc widget types, so if we dont have one dont set the parameter
+						if ((SpecificWidgetType.IsValid() == true) && (SpecificWidgetType != NAME_None))
+						{
+							StructPropertyHandle->GetChildHandle(GET_MEMBER_NAME_CHECKED(FTutorialContentAnchor, TabToFocusOrOpen))->SetValue(TutorialMeta->TabTypeToOpen);
+						}
+					}
+					else
+					{
+						TSharedRef<FTagMetaData> GraphNodeMeta = StaticCastSharedRef<FTagMetaData>(MetaDataEntry);
+						FriendlyNameToSet = GraphNodeMeta->Tag.ToString();
+					}
+
+					// Set the friendly name to the PickedWidget name - we might not have any metadata
+					if (FriendlyNameProperty.IsValid())
+					{
+						FriendlyNameProperty->SetValue(FriendlyNameToSet);
+					}
 				}
 			}
 			
+			// Reset the pick data
 			PickedWidgetName = NAME_None;
 			PickedAllMetaData.Reset();
 
@@ -306,9 +366,15 @@ private:
 	}
 
 private:
-	/** Handle to the property struct */
+	/** Handle to the property struct (Data other than the name is in here) */
 	TSharedPtr<class IPropertyHandle> StructPropertyHandle;
 
+	/** Handle to the name property struct. This is the property we are actually picking */
+	TSharedPtr<class IPropertyHandle> PickPropertyHandle;
+
+	/* Handle to the friendly name property we should set if any. */
+	TSharedPtr<class IPropertyHandle> FriendlyNameProperty;
+
 	/** Handle to the window that contains this widget */
 	TWeakPtr<SWindow> ParentWindow;
 
@@ -317,6 +383,9 @@ private:
 
 	/* The metadata for the widget we are picking */
 	TArray<TSharedRef<FTagMetaData>> PickedAllMetaData;
+
+	/* If we are we picking a specific widget type this will specify a typename (EG SDockTab) */
+	FName SpecificWidgetType;
 };
 
 /** Widget used to launch a 'picking' session */
@@ -325,6 +394,10 @@ class SWidgetPicker : public SCompoundWidget
 public:
 	SLATE_BEGIN_ARGS( SWidgetPicker ) {}
 
+	SLATE_ARGUMENT(FName, SpecificWidgetType)
+
+	SLATE_ARGUMENT(TSharedPtr<class IPropertyHandle>, FriendlyNameProperty)
+
 	SLATE_END_ARGS()
 
 	~SWidgetPicker()
@@ -338,9 +411,13 @@ public:
 		}
 	}
 
-	void Construct(const FArguments& InArgs, TSharedRef<class IPropertyHandle> InStructPropertyHandle, IPropertyTypeCustomizationUtils* InStructCustomizationUtils)
+	void Construct(const FArguments& InArgs, TSharedRef<class IPropertyHandle> InStructPropertyHandle, TSharedRef<class IPropertyHandle> InPickPropertyHandle, IPropertyTypeCustomizationUtils* InStructCustomizationUtils)
 	{
 		StructPropertyHandle = InStructPropertyHandle;
+		PickPropertyHandle = InPickPropertyHandle;
+
+		SpecificWidgetType = InArgs._SpecificWidgetType;
+		FriendlyNameProperty = InArgs._FriendlyNameProperty;
 
 		ChildSlot
 		[
@@ -383,8 +460,10 @@ private:
 			PickerWindow = NewWindow;
 
 			NewWindow->SetContent(
-				SAssignNew(PickerWidget, SWidgetPickerFloatingWindow, StructPropertyHandle.ToSharedRef())
+				SAssignNew(PickerWidget, SWidgetPickerFloatingWindow, StructPropertyHandle.ToSharedRef(), PickPropertyHandle.ToSharedRef())
 				.ParentWindow(NewWindow)
+				.SpecificWidgetType(SpecificWidgetType)
+				.FriendlyNameProperty(FriendlyNameProperty)
 				);
 
 			TSharedPtr<SWindow> RootWindow = FGlobalTabmanager::Get()->GetRootWindow();
@@ -399,6 +478,7 @@ private:
 
 			FIntroTutorials& IntroTutorials = FModuleManager::Get().GetModuleChecked<FIntroTutorials>("IntroTutorials");
 			IntroTutorials.OnIsPicking().BindSP(this, &SWidgetPicker::OnIsPicking);
+			IntroTutorials.OnValidatePickingCandidate().BindSP(this, &SWidgetPicker::OnValidatePickingCandidate);
 		}
 
 		return FReply::Handled();
@@ -406,13 +486,24 @@ private:
 
 	FText HandlePickerStatusText() const
 	{
-		FName WidgetName = NAME_None;
-		TSharedPtr<IPropertyHandle> WrapperIdentifierProperty = StructPropertyHandle->GetChildHandle(GET_MEMBER_NAME_CHECKED(FTutorialContentAnchor, WrapperIdentifier));
+		if (PickPropertyHandle.IsValid())
+		{
+			FString WidgetValue;
+			PickPropertyHandle->GetValue(WidgetValue);
 
-		// Some of the tags contain a GUID, we dont want to see that here
-		FString WidgetValue;
-		WrapperIdentifierProperty->GetValue(WidgetValue);
-		return MakeFriendlyStringFromName(WidgetValue);
+			if (FriendlyNameProperty.IsValid())
+			{
+				FString FriendlyName;
+				FriendlyNameProperty->GetValue(FriendlyName);
+				if (FriendlyName.IsEmpty() == false)
+				{
+					WidgetValue = FriendlyName;
+				}
+			}
+
+			return FText::FromString(WidgetValue);
+		}
+		return FText();
 	}
 
 	FText MakeFriendlyStringFromName(const FString& WidgetName) const
@@ -442,6 +533,39 @@ private:
 		
 		return false;
 	}
+		
+	bool OnValidatePickingCandidate(TSharedRef<SWidget> InWidget, FName& OutWidgetNameToHighlight, bool& bOutShouldHighlight) const
+	{
+		bool bIsPicking = false;
+		bOutShouldHighlight = false;
+		TSharedPtr<FTagMetaData> MetaData = InWidget->GetMetaData<FTagMetaData>();
+		const FName Tag = (MetaData.IsValid() && MetaData->Tag.IsValid()) ? MetaData->Tag : InWidget->GetTag();
+		if (PickerWidget.IsValid())
+		{
+			// Is the given widget a candidate
+			if ((Tag != NAME_None) )
+			{
+				bIsPicking = true;
+			}
+			if (SpecificWidgetType != NAME_None)
+			{
+				bIsPicking = InWidget->GetType() == SpecificWidgetType;
+			}
+			
+			// If we are picking a widget check if we should also highlight it
+			if (bIsPicking == true)
+			{
+				OutWidgetNameToHighlight = PickerWidget.Pin()->GetPickedWidgetName();
+				FName InPickableName = PickerWidget.Pin()->GetPickableNameForWidget(InWidget);
+
+				if (InPickableName == OutWidgetNameToHighlight)
+				{
+					bOutShouldHighlight = true;
+				}
+			}
+		}
+		return bIsPicking;
+	}
 
 private:
 	/** Picker window widget */
@@ -452,6 +576,15 @@ private:
 
 	/** Handle to the struct we are customizing */
 	TSharedPtr<class IPropertyHandle> StructPropertyHandle;
+
+	/** Handle to the property we are customizing */
+	TSharedPtr<class IPropertyHandle> PickPropertyHandle;
+
+	/* Handle to the friendly name property we should set  if any. */
+	TSharedPtr<class IPropertyHandle> FriendlyNameProperty;
+
+	/* Are we picking a specific widget type */
+	FName	SpecificWidgetType;
 };
 
 TSharedRef<IPropertyTypeCustomization> FTutorialContentAnchorCustomization::MakeInstance()
@@ -462,6 +595,9 @@ TSharedRef<IPropertyTypeCustomization> FTutorialContentAnchorCustomization::Make
 void FTutorialContentAnchorCustomization::CustomizeHeader( TSharedRef<class IPropertyHandle> InStructPropertyHandle, class FDetailWidgetRow& HeaderRow, IPropertyTypeCustomizationUtils& StructCustomizationUtils )
 {
 	TSharedPtr<IPropertyHandle> DrawHighlightProperty = InStructPropertyHandle->GetChildHandle(GET_MEMBER_NAME_CHECKED(FTutorialContentAnchor, bDrawHighlight));
+	TSharedPtr<IPropertyHandle> WidgetNameProperty = InStructPropertyHandle->GetChildHandle(GET_MEMBER_NAME_CHECKED(FTutorialContentAnchor, WrapperIdentifier));
+	TSharedPtr<IPropertyHandle> TabToFocusProperty = InStructPropertyHandle->GetChildHandle(GET_MEMBER_NAME_CHECKED(FTutorialContentAnchor, TabToFocusOrOpen));
+	TSharedPtr<IPropertyHandle> FriendlyNameProperty = InStructPropertyHandle->GetChildHandle(GET_MEMBER_NAME_CHECKED(FTutorialContentAnchor, FriendlyName));
 
 	HeaderRow
 	.NameContent()
@@ -479,6 +615,12 @@ void FTutorialContentAnchorCustomization::CustomizeHeader( TSharedRef<class IPro
 		[
 			DrawHighlightProperty->CreatePropertyNameWidget()
 		]
+		+ SVerticalBox::Slot()
+		.AutoHeight()
+		.Padding(0.0f, 2.0f)
+		[
+			TabToFocusProperty->CreatePropertyNameWidget()
+		]
 	]
 	.ValueContent()
 	.MinDesiredWidth(250.0f)
@@ -489,7 +631,8 @@ void FTutorialContentAnchorCustomization::CustomizeHeader( TSharedRef<class IPro
 		.AutoHeight()
 		.Padding(0.0f, 2.0f)
 		[
-			SNew(SWidgetPicker, InStructPropertyHandle, &StructCustomizationUtils)
+			SNew(SWidgetPicker, InStructPropertyHandle, WidgetNameProperty.ToSharedRef(), &StructCustomizationUtils)
+			.FriendlyNameProperty(FriendlyNameProperty)
 		]
 		+ SVerticalBox::Slot()
 		.AutoHeight()
@@ -497,6 +640,13 @@ void FTutorialContentAnchorCustomization::CustomizeHeader( TSharedRef<class IPro
 		[
 			DrawHighlightProperty->CreatePropertyValueWidget()
 		]
+		+ SVerticalBox::Slot()
+		.AutoHeight()
+		.Padding(0.0f, 2.0f)
+		[
+			SNew(SWidgetPicker, InStructPropertyHandle, TabToFocusProperty.ToSharedRef(), &StructCustomizationUtils)
+				.SpecificWidgetType("SDockTab")
+		]
 	];
 }
 
diff --git a/Engine/Source/Editor/IntroTutorials/Public/EditorTutorial.h b/Engine/Source/Editor/IntroTutorials/Public/EditorTutorial.h
--- a/Engine/Source/Editor/IntroTutorials/Public/EditorTutorial.h
+++ b/Engine/Source/Editor/IntroTutorials/Public/EditorTutorial.h
@@ -123,15 +123,15 @@ struct INTROTUTORIALS_API FTutorialContentAnchor
 	UPROPERTY(EditAnywhere, Category="Anchor")
 	bool bDrawHighlight;
 	
+	/* Tab on which to focus (EG 'My Blueprint' tab). */
+	UPROPERTY(EditAnywhere, Category = "Anchor")
+	FString TabToFocusOrOpen;
+
+	//FBlueprintGraphNodeMetaData specific members	
 	/* User friendly name to display in the dialog */
 	UPROPERTY(EditAnywhere, Category = "AnchorMeta|FBlueprintGraphNodeMetaData")
 	FString FriendlyName;
 
-	/* The type of tab to open for this node*/
-	UPROPERTY(EditAnywhere, Category = "AnchorMeta|FBlueprintGraphNodeMetaData")
-	FString TabTypeToOpen;
-
-	//FBlueprintGraphNodeMetaData specific members
 	/* The GUID string */
 	UPROPERTY(EditAnywhere, Category = "AnchorMeta|FBlueprintGraphNodeMetaData")
  	FString GUIDString;
diff --git a/Engine/Source/Runtime/Slate/Private/Framework/Docking/TabManager.cpp b/Engine/Source/Runtime/Slate/Private/Framework/Docking/TabManager.cpp
--- a/Engine/Source/Runtime/Slate/Private/Framework/Docking/TabManager.cpp
+++ b/Engine/Source/Runtime/Slate/Private/Framework/Docking/TabManager.cpp
@@ -1183,11 +1183,7 @@ TSharedRef<SDockTab> FTabManager::SpawnTab( const FTabId& TabId, const TSharedPt
 		NewTabWidget->SetLayoutIdentifier( TabId );
 		NewTabWidget->ProvideDefaultLabel( Spawner->GetDisplayName().IsEmpty() ? FText::FromName( Spawner->TabType ) : Spawner->GetDisplayName() );
 		NewTabWidget->ProvideDefaultIcon( Spawner->GetIcon().GetIcon() );
-
-		// If this tabs' content was not tagged for tutorials by its creator, we can provide an automated tag based  on the TabId.
-		// NOTE: relying on this is bad because it couples tutorials to the internals of tab management!
-		NewTabWidget->GetContent()->AddMetadataIfMissing<FTagMetaData>( MakeShareable(new FTagMetaData( *TabId.ToString() ) ) );
-
+	
 		// The spawner tracks that last tab it spawned
 		Spawner->SpawnedTabPtr = NewTabWidget;
 	}
diff --git a/Engine/Source/Runtime/SlateCore/Public/Widgets/SWidget.h b/Engine/Source/Runtime/SlateCore/Public/Widgets/SWidget.h
--- a/Engine/Source/Runtime/SlateCore/Public/Widgets/SWidget.h
+++ b/Engine/Source/Runtime/SlateCore/Public/Widgets/SWidget.h
@@ -665,32 +665,6 @@ public:
 		return FoundMetaData;
 	}
 
-private:
-	/**
-	 * Add metadata to the widget unless this type of metadata is already present.
-	 *
-	 * MIGHT BE REMOVED ! ! !
-	 *
-	 * Not sure we want to allow such mutation of metadata yet.
-	 * We might remove this method, so it is only exposed to the
-	 * one used case at the moment; thus the 'friend FTabManager'
-	 */
-	friend class FTabManager;
-
-	template<typename MetaDataType>
-	void AddMetadataIfMissing( const TSharedRef<MetaDataType>& AddMe )
-	{
-		const bool bMetadataAlreadyPresent = MetaData.ContainsByPredicate([]( const TSharedRef<ISlateMetaData>& Candidate )
-		{
-			return Candidate->IsOfType<MetaDataType>();
-		});
-
-		if (!bMetadataAlreadyPresent)
-		{
-			MetaData.Add( AddMe );
-		}		
-	}
-
 public:
 
 	// Widget Inspector and debugging methods
