Added code to save current tutorial progress when tutorials are closed. (Fixes recurring welcome bubble)

diff --git a/Engine/Source/Editor/IntroTutorials/Private/STutorialRoot.cpp b/Engine/Source/Editor/IntroTutorials/Private/STutorialRoot.cpp
--- a/Engine/Source/Editor/IntroTutorials/Private/STutorialRoot.cpp
+++ b/Engine/Source/Editor/IntroTutorials/Private/STutorialRoot.cpp
@@ -305,6 +305,11 @@ void STutorialRoot::HandleCloseClicked()
 	if(CurrentTutorial != nullptr)
 	{
 		CurrentTutorial->HandleTutorialClosed();
+		// Update the current stage when we close
+		bool bHaveSeenTutorial = false;
+		CurrentTutorialStage = GetDefault<UTutorialStateSettings>()->GetProgress(CurrentTutorial, bHaveSeenTutorial);
+		GetMutableDefault<UTutorialStateSettings>()->RecordProgress(CurrentTutorial, CurrentTutorialStage);
+		GetMutableDefault<UTutorialStateSettings>()->SaveProgress();
 	}
 
 	// submit analytics data
