#UE4 DOCS: Copied the soft reference functionality from the Global Editor Utility class to the Editor Tutorial Class

diff --git a/Engine/Source/Editor/IntroTutorials/Private/EditorTutorial.cpp b/Engine/Source/Editor/IntroTutorials/Private/EditorTutorial.cpp
--- a/Engine/Source/Editor/IntroTutorials/Private/EditorTutorial.cpp
+++ b/Engine/Source/Editor/IntroTutorials/Private/EditorTutorial.cpp
@@ -73,3 +73,11 @@ void UEditorTutorial::OpenAsset(UObject* Asset)
 	FAssetEditorManager::Get().OpenEditorForAsset(Asset);
 }
 
+AActor* UEditorTutorial::GetActorReference(FString PathToActor)
+{
+#if WITH_EDITOR
+	return Cast<AActor>(StaticFindObject(AActor::StaticClass(), GEditor->GetEditorWorldContext().World(), *PathToActor, false));
+#else
+	return nullptr;
+#endif //WITH_EDITOR
+}
\ No newline at end of file
diff --git a/Engine/Source/Editor/IntroTutorials/Public/EditorTutorial.h b/Engine/Source/Editor/IntroTutorials/Public/EditorTutorial.h
--- a/Engine/Source/Editor/IntroTutorials/Public/EditorTutorial.h
+++ b/Engine/Source/Editor/IntroTutorials/Public/EditorTutorial.h
@@ -239,6 +239,14 @@ class INTROTUTORIALS_API UEditorTutorial : public UObject
 	UPROPERTY(EditAnywhere, Category="Tutorial", meta=(MetaClass="EditorTutorial"))
 	FStringClassReference NextTutorial;
 
+	/**
+	* Attempts to find the actor specified by PathToActor in the current editor world
+	* @param	PathToActor	The path to the actor (e.g. PersistentLevel.PlayerStart)
+	* @return	A reference to the actor, or none if it wasn't found
+	*/
+	UFUNCTION(BlueprintPure, Category = "Tutorial")
+	AActor* GetActorReference(FString PathToActor);
+
 	/** A standalone tutorial displays no navigation buttons and each content widget has a close button */
 	UPROPERTY(EditAnywhere, Category="Tutorial")
 	bool bIsStandalone;
