Restructured tutorial analytics events

Where dynamic event names were being constructed from tutorial assets, now reworked events to add a dynamic parameter.

UE-4978 - Restructure Tutorial Analytic Events

diff --git a/Engine/Source/Editor/IntroTutorials/Private/IntroTutorials.cpp b/Engine/Source/Editor/IntroTutorials/Private/IntroTutorials.cpp
--- a/Engine/Source/Editor/IntroTutorials/Private/IntroTutorials.cpp
+++ b/Engine/Source/Editor/IntroTutorials/Private/IntroTutorials.cpp
@@ -50,7 +50,7 @@ FIntroTutorials::FIntroTutorials()
 	bDesireResettingTutorialSeenFlagOnLoad = FParse::Param(FCommandLine::Get(), TEXT("ResetTutorials"));
 }
 
-FString FIntroTutorials::AnalyticsEventNameFromTutorial(const FString& BaseEventName, UEditorTutorial* Tutorial)
+FString FIntroTutorials::AnalyticsEventNameFromTutorial(UEditorTutorial* Tutorial)
 {
 	FString TutorialPath = Tutorial->GetOutermost()->GetFName().ToString();
 
@@ -58,12 +58,7 @@ FString FIntroTutorials::AnalyticsEventNameFromTutorial(const FString& BaseEvent
 	FString RightStr;
 	TutorialPath.Split( TEXT("/"), NULL, &RightStr, ESearchCase::IgnoreCase, ESearchDir::FromEnd );
 
-	// then append that to the header
-	// e.g. Rocket.Tutorials.ClosedInEditorTutorial
-	FString OutStr = BaseEventName;
-	OutStr += RightStr;
-
-	return OutStr;
+	return RightStr;
 }
 
 TSharedRef<FExtender> FIntroTutorials::AddSummonBlueprintTutorialsMenuExtender(const TSharedRef<FUICommandList> CommandList, const TArray<UObject*> EditingObjects) const
diff --git a/Engine/Source/Editor/IntroTutorials/Private/IntroTutorials.h b/Engine/Source/Editor/IntroTutorials/Private/IntroTutorials.h
--- a/Engine/Source/Editor/IntroTutorials/Private/IntroTutorials.h
+++ b/Engine/Source/Editor/IntroTutorials/Private/IntroTutorials.h
@@ -33,7 +33,7 @@ public:
 
 	void SummonTutorialBrowser();
 
-	static FString AnalyticsEventNameFromTutorial(const FString& BaseEventName, UEditorTutorial* Tutorial);
+	static FString AnalyticsEventNameFromTutorial(UEditorTutorial* Tutorial);
 
 private:
 
diff --git a/Engine/Source/Editor/IntroTutorials/Private/STutorialContent.cpp b/Engine/Source/Editor/IntroTutorials/Private/STutorialContent.cpp
--- a/Engine/Source/Editor/IntroTutorials/Private/STutorialContent.cpp
+++ b/Engine/Source/Editor/IntroTutorials/Private/STutorialContent.cpp
@@ -619,9 +619,12 @@ void STutorialContent::HandleRestartSelected()
 		const bool bRestart = true;
 		IntroTutorials.LaunchTutorial(Tutorial.Get(), bRestart, FSlateApplication::Get().FindWidgetWindow(AsShared()));
 
-		if( FEngineAnalytics::IsAvailable() )
+		if( FEngineAnalytics::IsAvailable() && Tutorial.IsValid() )
 		{
-			FEngineAnalytics::GetProvider().RecordEvent( FIntroTutorials::AnalyticsEventNameFromTutorial(TEXT("Rocket.Tutorials.Restarted"), Tutorial.Get()) );
+			TArray<FAnalyticsEventAttribute> EventAttributes;
+			EventAttributes.Add(FAnalyticsEventAttribute(TEXT("TutorialAsset"), FIntroTutorials::AnalyticsEventNameFromTutorial(Tutorial.Get())));
+
+			FEngineAnalytics::GetProvider().RecordEvent( TEXT("Rocket.Tutorials.Restarted"), EventAttributes );
 		}
 	}
 }
@@ -631,7 +634,7 @@ void STutorialContent::HandleBrowseSelected()
 	if( FEngineAnalytics::IsAvailable() && Tutorial.IsValid())
 	{
 		TArray<FAnalyticsEventAttribute> EventAttributes;
-		EventAttributes.Add(FAnalyticsEventAttribute(TEXT("FromTutorial"), FIntroTutorials::AnalyticsEventNameFromTutorial(TEXT(""), Tutorial.Get())));
+		EventAttributes.Add(FAnalyticsEventAttribute(TEXT("FromTutorial"), FIntroTutorials::AnalyticsEventNameFromTutorial(Tutorial.Get())));
 
 		FEngineAnalytics::GetProvider().RecordEvent( TEXT("Rocket.Tutorials.OpenedBrowser"), EventAttributes );
 	}
diff --git a/Engine/Source/Editor/IntroTutorials/Private/STutorialRoot.cpp b/Engine/Source/Editor/IntroTutorials/Private/STutorialRoot.cpp
--- a/Engine/Source/Editor/IntroTutorials/Private/STutorialRoot.cpp
+++ b/Engine/Source/Editor/IntroTutorials/Private/STutorialRoot.cpp
@@ -178,7 +178,7 @@ void STutorialRoot::HandleBackClicked()
 	if( FEngineAnalytics::IsAvailable() && CurrentTutorial != nullptr)
 	{
 		TArray<FAnalyticsEventAttribute> EventAttributes;
-		EventAttributes.Add(FAnalyticsEventAttribute(TEXT("Context.Tutorial"), FIntroTutorials::AnalyticsEventNameFromTutorial(TEXT(""), CurrentTutorial)));
+		EventAttributes.Add(FAnalyticsEventAttribute(TEXT("Context.Tutorial"), FIntroTutorials::AnalyticsEventNameFromTutorial(CurrentTutorial)));
 		EventAttributes.Add(FAnalyticsEventAttribute(TEXT("Context.StageIndex"), CurrentTutorialStage));
 
 		FEngineAnalytics::GetProvider().RecordEvent( TEXT("Rocket.Tutorials.ClickedBackButton"), EventAttributes );
@@ -331,8 +331,9 @@ void STutorialRoot::HandleCloseClicked()
 		EventAttributes.Add(FAnalyticsEventAttribute(TEXT("LastStageIndex"), CurrentExcerptIndex));
 		EventAttributes.Add(FAnalyticsEventAttribute(TEXT("LastStageTitle"), CurrentExcerptTitle));
 		EventAttributes.Add(FAnalyticsEventAttribute(TEXT("TimeSpentInTutorial"), CurrentPageElapsedTime));
+		EventAttributes.Add(FAnalyticsEventAttribute(TEXT("TutorialAsset"), FIntroTutorials::AnalyticsEventNameFromTutorial(CurrentTutorial)));
 			
-		FEngineAnalytics::GetProvider().RecordEvent( FIntroTutorials::AnalyticsEventNameFromTutorial(TEXT("Rocket.Tutorials.Closed"), CurrentTutorial), EventAttributes );
+		FEngineAnalytics::GetProvider().RecordEvent( TEXT("Rocket.Tutorials.Closed"), EventAttributes );
 	}
 }
 
diff --git a/Engine/Source/Editor/IntroTutorials/Private/STutorialsBrowser.cpp b/Engine/Source/Editor/IntroTutorials/Private/STutorialsBrowser.cpp
--- a/Engine/Source/Editor/IntroTutorials/Private/STutorialsBrowser.cpp
+++ b/Engine/Source/Editor/IntroTutorials/Private/STutorialsBrowser.cpp
@@ -816,8 +816,9 @@ void STutorialsBrowser::OnTutorialSelected(UEditorTutorial* InTutorial, bool bRe
 	{
 		TArray<FAnalyticsEventAttribute> EventAttributes;
 		EventAttributes.Add(FAnalyticsEventAttribute(TEXT("Restarted"), bRestart));
+		EventAttributes.Add(FAnalyticsEventAttribute(TEXT("TutorialAsset"), FIntroTutorials::AnalyticsEventNameFromTutorial(InTutorial)));
 
-		FEngineAnalytics::GetProvider().RecordEvent( FIntroTutorials::AnalyticsEventNameFromTutorial(TEXT("Rocket.Tutorials.LaunchedFromBrowser"), InTutorial), EventAttributes );
+		FEngineAnalytics::GetProvider().RecordEvent( TEXT("Rocket.Tutorials.LaunchedFromBrowser"), EventAttributes );
 	}
 
 	OnLaunchTutorial.ExecuteIfBound(InTutorial, bRestart, ParentWindow, FSimpleDelegate(), FSimpleDelegate());
