Fixed stalls in tutorials

API break: added a new pure virtual to IAssetEditorInstance - InvokeTab().

Tutorials were attempting to load a bogus asset for editing for each stage, causing stalls each time the user clicked Next.
Re-arranged code that attempted to open asset editor tabs so that it wasnt attempting ot open a bogus asset (used the correct member).
Order of tab manager preference is now Per-stage Blueprint asset editor->Per-'tutorial asset' editor->level editor.
Added new function to allow us to open tabs in asset editors. This avoids a dangerous up-cast to FAssetEditorToolkit to access its tab manager.

diff --git a/Engine/Source/Editor/IntroTutorials/Private/SEditorTutorials.cpp b/Engine/Source/Editor/IntroTutorials/Private/SEditorTutorials.cpp
--- a/Engine/Source/Editor/IntroTutorials/Private/SEditorTutorials.cpp
+++ b/Engine/Source/Editor/IntroTutorials/Private/SEditorTutorials.cpp
@@ -28,6 +28,8 @@ void SEditorTutorials::Construct(const FArguments& InArgs)
 	[
 		ContentBox.ToSharedRef()
 	];
+
+	RebuildCurrentContent();
 }
 
 void SEditorTutorials::LaunchTutorial(bool bInIsNavigationWindow, FSimpleDelegate InOnTutorialClosed, FSimpleDelegate InOnTutorialExited)
@@ -43,6 +45,8 @@ void SEditorTutorials::HideContent()
 {
 	HandleHomeClicked();
 	bIsNavigationWindow = false;
+
+	RebuildCurrentContent();
 }
 
 bool SEditorTutorials::IsNavigationVisible() const
diff --git a/Engine/Source/Editor/IntroTutorials/Private/STutorialOverlay.cpp b/Engine/Source/Editor/IntroTutorials/Private/STutorialOverlay.cpp
--- a/Engine/Source/Editor/IntroTutorials/Private/STutorialOverlay.cpp
+++ b/Engine/Source/Editor/IntroTutorials/Private/STutorialOverlay.cpp
@@ -86,7 +86,7 @@ void STutorialOverlay::Construct(const FArguments& InArgs, UEditorTutorial* InTu
 						.OnWasWidgetDrawn(InArgs._OnWasWidgetDrawn)
 						.NextButtonText(InStage->NextButtonText);
 					
-					PerformWidgetInteractions(WidgetContent); 					
+					PerformWidgetInteractions(InTutorial, WidgetContent); 					
 
 					OverlayCanvas->AddSlot()
 						.Position(TAttribute<FVector2D>::Create(TAttribute<FVector2D>::FGetter::CreateSP(ContentWidget.Get(), &STutorialContent::GetPosition)))
@@ -175,69 +175,65 @@ int32 STutorialOverlay::TraverseWidgets(TSharedRef<SWidget> InWidget, const FGeo
 	return LayerId;
 }
 
-void STutorialOverlay::PerformWidgetInteractions(const FTutorialWidgetContent &WidgetContent)
+void STutorialOverlay::PerformWidgetInteractions(UEditorTutorial* InTutorial, const FTutorialWidgetContent &WidgetContent)
 {
 	// Open any browser we need too
-	OpenBrowserForWidgetAnchor(WidgetContent);
+	OpenBrowserForWidgetAnchor(InTutorial, WidgetContent);
 
 	FocusOnAnyBlueprintNodes(WidgetContent);
 }
 
-void STutorialOverlay::OpenBrowserForWidgetAnchor(const FTutorialWidgetContent &WidgetContent)
+void STutorialOverlay::OpenBrowserForWidgetAnchor(UEditorTutorial* InTutorial, const FTutorialWidgetContent &WidgetContent)
 {
-	//bool bTabOpened = false;
-	TSharedPtr<FTabManager> TabManager;
-	TArray<FString> AssetPaths;
-	FString Name;
-
-	// Try looking for the object via the ID first
-	Name = WidgetContent.WidgetAnchor.WrapperIdentifier.ToString();
-	AssetPaths.Add(Name);
-	// Opening the editor will force the object to be loaded if it exists
-	FAssetEditorManager::Get().OpenEditorsForAssets(AssetPaths);
-	UObject* AssetObject = FindObject<UObject>(ANY_PACKAGE, *Name);
-	IAssetEditorInstance* AssetEditor = FAssetEditorManager::Get().FindEditorForAsset(AssetObject, false);
-	// If we now have a valid asset editor get its tabmanager - we will use this to open/focus a tab if there is one
-	if (AssetEditor != nullptr)
+	if(!WidgetContent.WidgetAnchor.TabToFocusOrOpen.IsEmpty())
 	{
-		FAssetEditorToolkit* ToolkitEditor = static_cast<FAssetEditorToolkit*>(FAssetEditorManager::Get().FindEditorForAsset(AssetObject, false));
-		if ((ToolkitEditor != nullptr) && (WidgetContent.WidgetAnchor.TabToFocusOrOpen.IsEmpty() == false))
+		IAssetEditorInstance* AssetEditor = nullptr;
+
+		// Check to see if we can find a blueprint relevant to this node and open the editor for that (Then try to get the tabmanager from that)
+		if(WidgetContent.WidgetAnchor.OuterName.Len() > 0)
 		{
-			TabManager = ToolkitEditor->GetTabManager();
-		}
-	}
+			// Remove the prefix from the name
+			int32 Space = WidgetContent.WidgetAnchor.OuterName.Find(TEXT(" "));
+			FString Name = WidgetContent.WidgetAnchor.OuterName.RightChop(Space + 1);
+			TArray<FString> AssetPaths;
+			AssetPaths.Add(Name);
+			FAssetEditorManager::Get().OpenEditorsForAssets(AssetPaths);
+			UObject* Blueprint = FindObject<UObject>(ANY_PACKAGE, *Name);
 
-	// If we don't have a valid tab manager we should now check to see if we can find a blueprint relevant to this node and open the editor for that (Then try to get the tabmanager from that)
-	if (TabManager.IsValid() == false)
-	{
-		// Remove the prefix from the name
-		AssetPaths.Empty();
-		int32 Space = WidgetContent.WidgetAnchor.OuterName.Find(TEXT(" "));
-		Name = WidgetContent.WidgetAnchor.OuterName.RightChop(Space + 1);
-		AssetPaths.Add(Name);
-		FAssetEditorManager::Get().OpenEditorsForAssets(AssetPaths);
-		UObject* Blueprint = FindObject<UObject>(ANY_PACKAGE, *Name);
+			// If we found a blueprint
+			if(Blueprint != nullptr)
+			{
+				IAssetEditorInstance* PotentialAssetEditor = FAssetEditorManager::Get().FindEditorForAsset(Blueprint, false);
+				if(PotentialAssetEditor != nullptr)
+				{
+					AssetEditor = PotentialAssetEditor;
+				}
+			}
+		}
 
-		// If we found a blueprint
-		if (Blueprint != nullptr)
+		// If we haven't found a tab manager, next check the asset editor that we reference in this tutorial, if any
+		if(AssetEditor == nullptr)
 		{
-			IBlueprintEditor* BlueprintEditor = (FBlueprintEditor*)FAssetEditorManager::Get().FindEditorForAsset(Blueprint, false);
-			if ((BlueprintEditor != nullptr) && (WidgetContent.WidgetAnchor.TabToFocusOrOpen.IsEmpty() == false))
+			// Try looking for the object that this tutorial references (it should already be loaded by this tutorial if it exists).
+			UObject* AssetObject = InTutorial->AssetToUse.ResolveObject();
+			if(AssetObject != nullptr)
 			{
-				TabManager = BlueprintEditor->GetTabManager();
+				IAssetEditorInstance* PotentialAssetEditor = FAssetEditorManager::Get().FindEditorForAsset(AssetObject, false);
+				if(PotentialAssetEditor != nullptr)
+				{
+					AssetEditor = PotentialAssetEditor;
+				}
 			}
 		}
-	}
 
-	// Invoke any tab
-	if (WidgetContent.WidgetAnchor.TabToFocusOrOpen.IsEmpty() == false)
-	{
-		if (TabManager.IsValid() == true)
+		// Invoke any tab
+		if(AssetEditor != nullptr)
 		{
-			TSharedRef<SDockTab> IT = TabManager->InvokeTab(FTabId(*WidgetContent.WidgetAnchor.TabToFocusOrOpen));
+			AssetEditor->InvokeTab(FTabId(*WidgetContent.WidgetAnchor.TabToFocusOrOpen));
 		}
 		else
 		{
+			// fallback to trying the main level editor tab manager
 			FLevelEditorModule& LevelEditorModule = FModuleManager::GetModuleChecked<FLevelEditorModule>(TEXT("LevelEditor"));
 			TSharedPtr<FTabManager> LevelEditorTabManager = LevelEditorModule.GetLevelEditorTabManager();
 			LevelEditorTabManager->InvokeTab(FName(*WidgetContent.WidgetAnchor.TabToFocusOrOpen));
diff --git a/Engine/Source/Editor/IntroTutorials/Private/STutorialOverlay.h b/Engine/Source/Editor/IntroTutorials/Private/STutorialOverlay.h
--- a/Engine/Source/Editor/IntroTutorials/Private/STutorialOverlay.h
+++ b/Engine/Source/Editor/IntroTutorials/Private/STutorialOverlay.h
@@ -60,13 +60,13 @@ private:
 	int32 TraverseWidgets(TSharedRef<SWidget> InWidget, const FGeometry& InGeometry, const FSlateRect& MyClippingRect, FSlateWindowElementList& OutDrawElements, int32 LayerId) const;
 
 	/* Opens the browser that the given widget requires if it is not already. */
-	void OpenBrowserForWidgetAnchor(const FTutorialWidgetContent &WidgetContent);
+	void OpenBrowserForWidgetAnchor(UEditorTutorial* InTutorial, const FTutorialWidgetContent &WidgetContent);
 
 	/* Focus on the blueprint node if an anchor widget references one */
 	void FocusOnAnyBlueprintNodes(const FTutorialWidgetContent &WidgetContent);
 
 	/* Do any interaction stuff for a widget - open browser, scroll to node etc */
-	void PerformWidgetInteractions(const FTutorialWidgetContent &WidgetContent);
+	void PerformWidgetInteractions(UEditorTutorial* InTutorial, const FTutorialWidgetContent &WidgetContent);
 
 private:
 	/** Reference to the canvas we use to position our content widgets */
diff --git a/Engine/Source/Editor/UnrealEd/Private/Toolkits/AssetEditorToolkit.cpp b/Engine/Source/Editor/UnrealEd/Private/Toolkits/AssetEditorToolkit.cpp
--- a/Engine/Source/Editor/UnrealEd/Private/Toolkits/AssetEditorToolkit.cpp
+++ b/Engine/Source/Editor/UnrealEd/Private/Toolkits/AssetEditorToolkit.cpp
@@ -319,6 +319,10 @@ bool FAssetEditorToolkit::CloseWindow()
 	return true;
 }
 
+void FAssetEditorToolkit::InvokeTab(const FTabId& TabId)
+{
+	GetTabManager()->InvokeTab(TabId);
+}
 
 TSharedPtr< IToolkitHost > FAssetEditorToolkit::GetPreviousWorldCentricToolkitHost()
 {
diff --git a/Engine/Source/Editor/UnrealEd/Public/MiniCurveEditor.h b/Engine/Source/Editor/UnrealEd/Public/MiniCurveEditor.h
--- a/Engine/Source/Editor/UnrealEd/Public/MiniCurveEditor.h
+++ b/Engine/Source/Editor/UnrealEd/Public/MiniCurveEditor.h
@@ -24,6 +24,8 @@ public:
 	virtual void FocusWindow(UObject* ObjectToFocusOn) override;
 	virtual bool CloseWindow() override;
 	virtual bool IsPrimaryEditor() const override { return true; };
+	virtual void InvokeTab(const struct FTabId& TabId) override {}
+
 private:
 
 	float ViewMinInput;
diff --git a/Engine/Source/Editor/UnrealEd/Public/Toolkits/AssetEditorManager.h b/Engine/Source/Editor/UnrealEd/Public/Toolkits/AssetEditorManager.h
--- a/Engine/Source/Editor/UnrealEd/Public/Toolkits/AssetEditorManager.h
+++ b/Engine/Source/Editor/UnrealEd/Public/Toolkits/AssetEditorManager.h
@@ -20,6 +20,7 @@ public:
 	virtual void FocusWindow(UObject* ObjectToFocusOn = NULL) = 0;
 	virtual bool CloseWindow() = 0;
 	virtual bool IsPrimaryEditor() const = 0;
+	virtual void InvokeTab(const struct FTabId& TabId) = 0;
 };
 
 
diff --git a/Engine/Source/Editor/UnrealEd/Public/Toolkits/AssetEditorToolkit.h b/Engine/Source/Editor/UnrealEd/Public/Toolkits/AssetEditorToolkit.h
--- a/Engine/Source/Editor/UnrealEd/Public/Toolkits/AssetEditorToolkit.h
+++ b/Engine/Source/Editor/UnrealEd/Public/Toolkits/AssetEditorToolkit.h
@@ -67,6 +67,7 @@ public:
 	virtual void FocusWindow(UObject* ObjectToFocusOn = NULL) override;
 	virtual bool CloseWindow() override;
 	virtual bool IsPrimaryEditor() const override { return true; };
+	virtual void InvokeTab(const struct FTabId& TabId) override;
 
 	/**
 	 * Fills in the supplied menu with commands for working with this asset file
