#UE4: Fixed the 4 editor tutorial events to be events again by removing the old const bit (talked to Mike Beach and Phillip Kavan about this being the correct way to fix this)

diff --git a/Engine/Source/Editor/IntroTutorials/Public/EditorTutorial.h b/Engine/Source/Editor/IntroTutorials/Public/EditorTutorial.h
--- a/Engine/Source/Editor/IntroTutorials/Public/EditorTutorial.h
+++ b/Engine/Source/Editor/IntroTutorials/Public/EditorTutorial.h
@@ -284,20 +284,20 @@ public:
 
 protected:
 	/** Event fired when a tutorial stage begins */
-	UFUNCTION(BlueprintImplementableEvent, Category="Tutorial")
-	void OnTutorialStageStarted(FName StageName) const;
+	UFUNCTION(BlueprintImplementableEvent, Category = "Tutorial")
+	void OnTutorialStageStarted(FName StageName);
 
 	/** Event fired when a tutorial stage ends */
-	UFUNCTION(BlueprintImplementableEvent, Category="Tutorial")
-	void OnTutorialStageEnded(FName StageName) const;
+	UFUNCTION(BlueprintImplementableEvent, Category = "Tutorial")
+	void OnTutorialStageEnded(FName StageName);
 
 	/** Event fired when a tutorial is launched */
 	UFUNCTION(BlueprintImplementableEvent, Category = "Tutorial")
-	void OnTutorialLaunched() const;
+	void OnTutorialLaunched();
 
 	/** Event fired when a tutorial is closed */
 	UFUNCTION(BlueprintImplementableEvent, Category = "Tutorial")
-	void OnTutorialClosed() const;
+	void OnTutorialClosed();
 
 	/** Advance to the next stage of a tutorial */
 	UFUNCTION(BlueprintCallable, Category="Tutorial")
