UE-7990 fixed the technical issues. Still should have a visual element, like a throbber, to give users some feedback as to why the tutorial window isn't popping up immediately.

diff --git a/Engine/Source/Editor/IntroTutorials/Private/STutorialButton.cpp b/Engine/Source/Editor/IntroTutorials/Private/STutorialButton.cpp
--- a/Engine/Source/Editor/IntroTutorials/Private/STutorialButton.cpp
+++ b/Engine/Source/Editor/IntroTutorials/Private/STutorialButton.cpp
@@ -6,6 +6,7 @@
 #include "TutorialStateSettings.h"
 #include "TutorialMetaData.h"
 #include "EngineBuildSettings.h"
+#include "AssetRegistryModule.h"
 
 #define LOCTEXT_NAMESPACE "STutorialButton"
 
@@ -22,6 +23,7 @@ void STutorialButton::Construct(const FArguments& InArgs)
 
 	bTestAlerts = FParse::Param(FCommandLine::Get(), TEXT("TestTutorialAlerts"));
 
+	bPendingClickAction = false;
 	bTutorialAvailable = false;
 	bTutorialCompleted = false;
 	bTutorialDismissed = false;
@@ -135,25 +137,57 @@ int32 STutorialButton::OnPaint(const FPaintArgs& Args, const FGeometry& Allotted
 
 FReply STutorialButton::HandleButtonClicked()
 {
+	if (bPendingClickAction)
+	{
+		//There's already a click pending
+		return FReply::Handled();
+	}
+
 	RefreshStatus();
 
-	if( FEngineAnalytics::IsAvailable() )
+	if (FEngineAnalytics::IsAvailable())
 	{
 		TArray<FAnalyticsEventAttribute> EventAttributes;
 		EventAttributes.Add(FAnalyticsEventAttribute(TEXT("Context"), Context.ToString()));
 		EventAttributes.Add(FAnalyticsEventAttribute(TEXT("TimeSinceAlertStarted"), (AlertStartTime != 0.0f && ShouldShowAlert()) ? (FPlatformTime::Seconds() - AlertStartTime) : -1.0f));
 		EventAttributes.Add(FAnalyticsEventAttribute(TEXT("LaunchedBrowser"), ShouldLaunchBrowser()));
 
-		FEngineAnalytics::GetProvider().RecordEvent( TEXT("Rocket.Tutorials.ClickedContextButton"), EventAttributes );
+		FEngineAnalytics::GetProvider().RecordEvent(TEXT("Rocket.Tutorials.ClickedContextButton"), EventAttributes);
 	}
 
+	bPendingClickAction = true;
+	RegisterActiveTimer(0.f, FWidgetActiveTimerDelegate::CreateSP(this, &STutorialButton::HandleButtonClicked_AssetRegistryChecker));
+	return FReply::Handled();
+}
+
+EActiveTimerReturnType STutorialButton::HandleButtonClicked_AssetRegistryChecker(double InCurrentTime, float InDeltaTime)
+{
+	//Force tutorials to load into the asset registry before we proceed any further.
+	FAssetRegistryModule& AssetRegistry = FModuleManager::LoadModuleChecked<FAssetRegistryModule>(TEXT("AssetRegistry"));
+	bool IsStillLoading = AssetRegistry.Get().IsLoadingAssets();
+	if (IsStillLoading)
+	{
+		//We could tick the asset registry here, but we don't need to.
+		return EActiveTimerReturnType::Continue;
+	}
+
+	//Sometimes, this gives a false positive because the tutorial we want to launch wasn't loaded into the asset registry when we checked. Opening and closing the tab works around that by letting the browser recheck.
+	if (ShouldLaunchBrowser())
+	{
+		FLevelEditorModule& LevelEditorModule = FModuleManager::GetModuleChecked<FLevelEditorModule>(TEXT("LevelEditor"));
+		LevelEditorModule.GetLevelEditorTabManager()->InvokeTab(FTabId("TutorialsBrowser"))->RequestCloseTab();
+		RefreshStatus();
+	}
+
+	//Now we know the asset registry is loaded, the tutorial broswer is updated, and we are ready to complete the click and stop this active timer
 	FIntroTutorials& IntroTutorials = FModuleManager::GetModuleChecked<FIntroTutorials>(TEXT("IntroTutorials"));
-	if(ShouldLaunchBrowser())
+	if (ShouldLaunchBrowser())
 	{
 		IntroTutorials.SummonTutorialBrowser();
 	}
 	else if (CachedLaunchTutorial != nullptr)
 	{
+		//If we don't want to launch the browser, and we have a tutorial in mind, launch the tutorial now.
 		auto Delegate = FSimpleDelegate::CreateSP(this, &STutorialButton::HandleTutorialExited);
 
 		const bool bRestart = true;
@@ -165,7 +199,8 @@ FReply STutorialButton::HandleButtonClicked()
 		bTutorialDismissed = true;
 	}
 
-	return FReply::Handled();
+	bPendingClickAction = false;
+	return EActiveTimerReturnType::Stop;
 }
 
 FReply STutorialButton::OnMouseButtonDown(const FGeometry& MyGeometry, const FPointerEvent& MouseEvent)
diff --git a/Engine/Source/Editor/IntroTutorials/Private/STutorialButton.h b/Engine/Source/Editor/IntroTutorials/Private/STutorialButton.h
--- a/Engine/Source/Editor/IntroTutorials/Private/STutorialButton.h
+++ b/Engine/Source/Editor/IntroTutorials/Private/STutorialButton.h
@@ -25,6 +25,9 @@ private:
 	/** Handle clicking the tutorial button */
 	FReply HandleButtonClicked();
 
+	/** After the initial click is processed, go here, wait for the asset registry to load, and complete the action once we're ready */
+	EActiveTimerReturnType HandleButtonClicked_AssetRegistryChecker(double InCurrentTime, float InDeltaTime);
+
 	/** Dismiss the pulsing alert */
 	void DismissAlert();
 
@@ -86,4 +89,7 @@ private:
 
 	/** Cached browser filter */
 	FString CachedBrowserFilter;
+
+	/** True if we're waiting for asset registry to load in response to a click */
+	bool bPendingClickAction;
 };
\ No newline at end of file
diff --git a/Engine/Source/Editor/IntroTutorials/Private/STutorialsBrowser.cpp b/Engine/Source/Editor/IntroTutorials/Private/STutorialsBrowser.cpp
--- a/Engine/Source/Editor/IntroTutorials/Private/STutorialsBrowser.cpp
+++ b/Engine/Source/Editor/IntroTutorials/Private/STutorialsBrowser.cpp
@@ -686,9 +686,16 @@ void STutorialsBrowser::RebuildTutorials(TSharedPtr<FTutorialListEntry_Category>
 {
 	TArray<TSharedPtr<FTutorialListEntry_Tutorial>> Tutorials;
 
-	// rebuild tutorials
+	//Ensure that tutorials are loaded into the asset registry before making a list of them.
 	FAssetRegistryModule& AssetRegistry = FModuleManager::LoadModuleChecked<FAssetRegistryModule>(TEXT("AssetRegistry"));
+	bool IsStillLoading = AssetRegistry.Get().IsLoadingAssets();
+	if (IsStillLoading)
+	{
+		//This can happen if you close the editor with the tutorials browser open, then restart the editor so that the browser opens immediately.
+		return;
+	}
 
+	// rebuild tutorials
 	FARFilter Filter;
 	Filter.ClassNames.Add(UBlueprint::StaticClass()->GetFName());
 	Filter.bRecursiveClasses = true;
