Adding an include file to fix a NonUnity GUBP issue.

diff --git a/Engine/Source/Editor/IntroTutorials/Private/STutorialButton.cpp b/Engine/Source/Editor/IntroTutorials/Private/STutorialButton.cpp
--- a/Engine/Source/Editor/IntroTutorials/Private/STutorialButton.cpp
+++ b/Engine/Source/Editor/IntroTutorials/Private/STutorialButton.cpp
@@ -7,6 +7,7 @@
 #include "TutorialMetaData.h"
 #include "EngineBuildSettings.h"
 #include "AssetRegistryModule.h"
+#include "LevelEditor.h"
 
 #define LOCTEXT_NAMESPACE "STutorialButton"
 
