Adding SDockTab.h to fix NonUnity compile.

diff --git a/Engine/Source/Editor/IntroTutorials/Private/STutorialButton.cpp b/Engine/Source/Editor/IntroTutorials/Private/STutorialButton.cpp
--- a/Engine/Source/Editor/IntroTutorials/Private/STutorialButton.cpp
+++ b/Engine/Source/Editor/IntroTutorials/Private/STutorialButton.cpp
@@ -8,6 +8,7 @@
 #include "EngineBuildSettings.h"
 #include "AssetRegistryModule.h"
 #include "LevelEditor.h"
+#include "SDockTab.h"
 
 #define LOCTEXT_NAMESPACE "STutorialButton"
 
