[INTEGRATE] Change 2411506 by [REDACTED_PERSON][REDACTED_HANDLE] on 2015/01/19 17:55:43

	Tutorial polish
	- Tweaked bubble animation speed and bounce effect
	- Increased padding between text and Next button (was too cramped)
	- Tweaked timing of pulsing green rectangle highlights

diff --git a/Engine/Source/Editor/IntroTutorials/Private/STutorialContent.cpp b/Engine/Source/Editor/IntroTutorials/Private/STutorialContent.cpp
--- a/Engine/Source/Editor/IntroTutorials/Private/STutorialContent.cpp
+++ b/Engine/Source/Editor/IntroTutorials/Private/STutorialContent.cpp
@@ -12,14 +12,14 @@
 
 namespace TutorialConstants
 {
-	const float BorderPulseAnimationLength = 1.0f;
-	const float BorderIntroAnimationLength = 0.5f;
-	const float ContentIntroAnimationLength = 0.5f;
+	const float BorderPulseAnimationLength = 0.75f;
+	const float BorderIntroAnimationLength = 0.4f;
+	const float ContentIntroAnimationLength = 0.25f;
 	const float MinBorderOpacity = 0.1f;
 	const float ShadowScale = 8.0f;
 	const float MaxBorderOffset = 8.0f;
 	const FMargin BorderSizeStandalone(24.0f, 24.0f);
-	const FMargin BorderSize(24.0f, 24.0f, 24.0f, 52.0f);
+	const FMargin BorderSize(24.0f, 24.0f, 24.0f, 62.0f);
 }
 
 const float ContentOffset = 10.0f;
