[INTEGRATE] Change 2412170 by [REDACTED_PERSON][REDACTED_HANDLE] on 2015/01/20 09:44:23

	Minor tutorial button tooltip cleanups

diff --git a/Engine/Source/Editor/IntroTutorials/Private/STutorialButton.cpp b/Engine/Source/Editor/IntroTutorials/Private/STutorialButton.cpp
--- a/Engine/Source/Editor/IntroTutorials/Private/STutorialButton.cpp
+++ b/Engine/Source/Editor/IntroTutorials/Private/STutorialButton.cpp
@@ -225,8 +225,8 @@ FReply STutorialButton::OnMouseButtonDown(const FGeometry& MyGeometry, const FPo
 		if(bTutorialAvailable)
 		{
 			MenuBuilder.AddMenuEntry(
-				FText::Format(LOCTEXT("LaunchTutorialPattern", "Open Tutorial: {0}"), TutorialTitle),
-				LOCTEXT("LaunchTutorialTooltip", "Launch this tutorial"),
+				FText::Format(LOCTEXT("LaunchTutorialPattern", "Start Tutorial: {0}"), TutorialTitle),
+				FText::Format(LOCTEXT("TutorialLaunchToolTip", "Click to begin the '{0}' tutorial"), TutorialTitle),
 				FSlateIcon(),
 				FUIAction(FExecuteAction::CreateSP(this, &STutorialButton::LaunchTutorial))
 				);
@@ -305,7 +305,7 @@ FText STutorialButton::GetButtonToolTip() const
 	}
 	else if(bTutorialAvailable)
 	{
-		return FText::Format(LOCTEXT("TutorialLaunchToolTipPattern", "Open: {0}\nRight-Click for More Options"), TutorialTitle);
+		return FText::Format(LOCTEXT("TutorialLaunchToolTipPattern", "Click to begin the '{0}' tutorial, or right click for more options"), TutorialTitle);
 	}
 	
 	return LOCTEXT("TutorialToolTip", "Take Tutorial");
