[INTEGRATE] Change 2413979 by bruce.nesbit@bnesbit_Releases on 2015/01/21 12:31:44

	Added support to tutorials to allow focus on content browser asset.

diff --git a/Engine/Source/Editor/IntroTutorials/Private/STutorialOverlay.cpp b/Engine/Source/Editor/IntroTutorials/Private/STutorialOverlay.cpp
--- a/Engine/Source/Editor/IntroTutorials/Private/STutorialOverlay.cpp
+++ b/Engine/Source/Editor/IntroTutorials/Private/STutorialOverlay.cpp
@@ -264,5 +264,16 @@ void STutorialOverlay::FocusOnAnyBlueprintNodes(const FTutorialWidgetContent &Wi
 			FKismetEditorUtilities::BringKismetToFocusAttentionOnObject(GraphNode, false);
 		}
 	}
+	else
+	{
+		// if we didn't have a blueprint object to focus on, try it with a regular one
+		UObject* FocusObject = FindObject<UObject>(ANY_PACKAGE, *WidgetContent.WidgetAnchor.WrapperIdentifier.ToString());
+		if (FocusObject != nullptr)
+		{
+			TArray< UObject* > Objects;
+			Objects.Add(FocusObject);
+			GEditor->SyncBrowserToObjects(Objects);
+		}
+	}
 }
 
