[INTEGRATE] Change 2414064 by bruce.nesbit@bnesbit_Releases on 2015/01/21 13:26:29

	Added IsValid check to new focus on content code.

diff --git a/Engine/Source/Editor/IntroTutorials/Private/STutorialOverlay.cpp b/Engine/Source/Editor/IntroTutorials/Private/STutorialOverlay.cpp
--- a/Engine/Source/Editor/IntroTutorials/Private/STutorialOverlay.cpp
+++ b/Engine/Source/Editor/IntroTutorials/Private/STutorialOverlay.cpp
@@ -264,7 +264,7 @@ void STutorialOverlay::FocusOnAnyBlueprintNodes(const FTutorialWidgetContent &Wi
 			FKismetEditorUtilities::BringKismetToFocusAttentionOnObject(GraphNode, false);
 		}
 	}
-	else
+	else /*if ( WidgetContent.WidgetAnchor.WrapperIdentifier.IsValid() == true )*/
 	{
 		// if we didn't have a blueprint object to focus on, try it with a regular one
 		UObject* FocusObject = FindObject<UObject>(ANY_PACKAGE, *WidgetContent.WidgetAnchor.WrapperIdentifier.ToString());
