[INTEGRATE] Change 2418658 by [REDACTED_PERSON][REDACTED_HANDLE] on 2015/01/26 11:10:25

	Added a "Play" button to editor tutorial assets in the Content Browser

	UE-7970 - Add small 'play' button to EditorTutorial asset thumbnail

	"EditorTutorial" now provides some class type actions which show a "Play" button on Blueprints that derive from it.

diff --git a/Engine/Source/Editor/EditorStyle/Private/SlateEditorStyle.cpp b/Engine/Source/Editor/EditorStyle/Private/SlateEditorStyle.cpp
--- a/Engine/Source/Editor/EditorStyle/Private/SlateEditorStyle.cpp
+++ b/Engine/Source/Editor/EditorStyle/Private/SlateEditorStyle.cpp
@@ -2689,6 +2689,7 @@ void FSlateEditorStyle::FStyle::SetupTutorialStyles()
 		Set( "Tutorials.Browser.DefaultCategoryIcon", new IMAGE_BRUSH( "Tutorials/DefaultCategoryIcon_40x", FVector2D(40.0f, 40.0f), FLinearColor::White ) );
 
 		Set( "Tutorials.Browser.BackButton.Image", new IMAGE_BRUSH( "Tutorials/BrowserBack", FVector2D(32.0f, 32.0f), FLinearColor(1.0f, 1.0f, 1.0f, 1.0f) ) );
+		Set( "Tutorials.Browser.PlayButton.Image", new IMAGE_BRUSH( "Tutorials/BrowserPlay", FVector2D(32.0f, 32.0f), FLinearColor(1.0f, 1.0f, 1.0f, 1.0f) ) );
 		Set( "Tutorials.Browser.RestartButton", new IMAGE_BRUSH( "Tutorials/BrowserRestart", FVector2D(16.0f, 16.0f), FLinearColor(1.0f, 1.0f, 1.0f, 1.0f) ) );
 
 		Set( "Tutorials.Browser.Completed", new IMAGE_BRUSH( "Tutorials/TutorialCompleted", Icon32x32 ) );
diff --git a/Engine/Source/Editor/IntroTutorials/Private/ClassTypeActions_EditorTutorial.cpp b/Engine/Source/Editor/IntroTutorials/Private/ClassTypeActions_EditorTutorial.cpp
new file mode 100644
--- /dev/null
+++ b/Engine/Source/Editor/IntroTutorials/Private/ClassTypeActions_EditorTutorial.cpp
@@ -0,0 +1,50 @@
+// Copyright 1998-2015 Epic Games, Inc. All Rights Reserved.
+
+#include "IntroTutorialsPrivatePCH.h"
+#include "ClassTypeActions_EditorTutorial.h"
+#include "EditorTutorial.h"
+
+#define LOCTEXT_NAMESPACE "IntroTutorials"
+
+UClass* FClassTypeActions_EditorTutorial::GetSupportedClass() const
+{
+	return UEditorTutorial::StaticClass();
+}
+
+TSharedPtr<SWidget> FClassTypeActions_EditorTutorial::GetThumbnailOverlay(const FAssetData& AssetData) const
+{
+	const FString FullTutorialAssetPath = AssetData.ObjectPath.ToString();
+	auto OnLaunchTutorialClicked = [FullTutorialAssetPath]() -> FReply
+	{
+		if (IIntroTutorials::IsAvailable())
+		{
+			IIntroTutorials& IntroTutorials = IIntroTutorials::Get();
+			IntroTutorials.LaunchTutorial(FullTutorialAssetPath);
+		}
+		return FReply::Handled();
+	};
+
+	return SNew(SBox)
+		.HAlign(HAlign_Center)
+		.VAlign(VAlign_Center)
+		.Padding(FMargin(2))
+		[
+			SNew(SButton)
+			.ButtonStyle(FEditorStyle::Get(), "HoverHintOnly")
+			.ToolTipText(LOCTEXT("Blueprint_LaunchTutorialToolTip", "Launch this tutorial."))
+			.Cursor(EMouseCursor::Default) // The outer widget can specify a DragHand cursor, so we need to override that here
+			.ForegroundColor(FSlateColor::UseForeground())
+			.OnClicked_Lambda(OnLaunchTutorialClicked)
+			[
+				SNew(SBox)
+				.MinDesiredWidth(16)
+				.MinDesiredHeight(16)
+				[
+					SNew(SImage)
+					.Image(FEditorStyle::GetBrush("Tutorials.Browser.PlayButton.Image"))
+				]
+			]
+		];
+}
+
+#undef LOCTEXT_NAMESPACE
diff --git a/Engine/Source/Editor/IntroTutorials/Private/ClassTypeActions_EditorTutorial.h b/Engine/Source/Editor/IntroTutorials/Private/ClassTypeActions_EditorTutorial.h
new file mode 100644
--- /dev/null
+++ b/Engine/Source/Editor/IntroTutorials/Private/ClassTypeActions_EditorTutorial.h
@@ -0,0 +1,13 @@
+// Copyright 1998-2015 Epic Games, Inc. All Rights Reserved.
+
+#pragma once
+
+#include "ClassTypeActions_Base.h"
+
+class FClassTypeActions_EditorTutorial : public FClassTypeActions_Base
+{
+public:
+	// IClassTypeActions Implementation
+	virtual UClass* GetSupportedClass() const override;
+	virtual TSharedPtr<class SWidget> GetThumbnailOverlay(const FAssetData& AssetData) const override;
+};
diff --git a/Engine/Source/Editor/IntroTutorials/Private/IntroTutorials.cpp b/Engine/Source/Editor/IntroTutorials/Private/IntroTutorials.cpp
--- a/Engine/Source/Editor/IntroTutorials/Private/IntroTutorials.cpp
+++ b/Engine/Source/Editor/IntroTutorials/Private/IntroTutorials.cpp
@@ -34,6 +34,8 @@
 #include "SDockTab.h"
 #include "ModuleManager.h"
 #include "IntroTutorials.h"
+#include "IAssetTools.h"
+#include "ClassTypeActions_EditorTutorial.h"
 
 
 #define LOCTEXT_NAMESPACE "IntroTutorials"
@@ -113,6 +115,14 @@ void FIntroTutorials::StartupModule()
 		{
 			GetMutableDefault<UTutorialStateSettings>()->ClearProgress();
 		}
+
+		// register our class actions to show the "Play" button on editor tutorial Blueprint assets
+		{
+			IAssetTools& AssetTools = FModuleManager::LoadModuleChecked<FAssetToolsModule>("AssetTools").Get();
+			TSharedRef<IClassTypeActions> EditorTutorialClassActions = MakeShareable(new FClassTypeActions_EditorTutorial());
+			RegisteredClassTypeActions.Add(EditorTutorialClassActions);
+			AssetTools.RegisterClassTypeActions(EditorTutorialClassActions);
+		}
 	}
 
 	// Register to display our settings
@@ -155,6 +165,16 @@ void FIntroTutorials::ShutdownModule()
 	if (!bDisableTutorials && !IsRunningCommandlet())
 	{
 		FSourceCodeNavigation::AccessOnCompilerNotFound().RemoveAll( this );
+
+		FAssetToolsModule* AssetToolsModule = FModuleManager::GetModulePtr<FAssetToolsModule>("AssetTools");
+		if (AssetToolsModule)
+		{
+			IAssetTools& AssetTools = AssetToolsModule->Get();
+			for(const auto& RegisteredClassTypeAction : RegisteredClassTypeActions)
+			{
+				AssetTools.UnregisterClassTypeActions(RegisteredClassTypeAction);
+			}
+		}
 	}
 
 	if (BlueprintEditorExtender.IsValid() && FModuleManager::Get().IsModuleLoaded("Kismet"))
diff --git a/Engine/Source/Editor/IntroTutorials/Private/IntroTutorials.h b/Engine/Source/Editor/IntroTutorials/Private/IntroTutorials.h
--- a/Engine/Source/Editor/IntroTutorials/Private/IntroTutorials.h
+++ b/Engine/Source/Editor/IntroTutorials/Private/IntroTutorials.h
@@ -12,6 +12,7 @@ DECLARE_DELEGATE_RetVal_ThreeParams(bool, FOnValidatePickingCandidate, TSharedRe
 class UEditorTutorial;
 struct FTutorialContent;
 class STutorialRoot;
+class IClassTypeActions;
 
 class FIntroTutorials : public IIntroTutorials
 {
@@ -128,4 +129,7 @@ private:
 
 	/** Curve asset for intros */
 	TWeakObjectPtr<UCurveFloat> ContentIntroCurve;
+
+	/** The collection of registered class type actions. */
+	TArray<TSharedRef<IClassTypeActions>> RegisteredClassTypeActions;
 };
\ No newline at end of file
diff --git a/Engine/Source/Editor/IntroTutorials/Public/IIntroTutorials.h b/Engine/Source/Editor/IntroTutorials/Public/IIntroTutorials.h
--- a/Engine/Source/Editor/IntroTutorials/Public/IIntroTutorials.h
+++ b/Engine/Source/Editor/IntroTutorials/Public/IIntroTutorials.h
@@ -23,7 +23,8 @@ public:
 	 */
 	static inline IIntroTutorials& Get()
 	{
-		return FModuleManager::LoadModuleChecked< IIntroTutorials >( "IntroTutorials" );
+		static const FName ModuleName = "IntroTutorials";
+		return FModuleManager::LoadModuleChecked< IIntroTutorials >( ModuleName );
 	}
 
 	/**
@@ -33,7 +34,8 @@ public:
 	 */
 	static inline bool IsAvailable()
 	{
-		return FModuleManager::Get().IsModuleLoaded( "IntroTutorials" );
+		static const FName ModuleName = "IntroTutorials";
+		return FModuleManager::Get().IsModuleLoaded( ModuleName );
 	}
 
 	/**
