UE-7990 - Fixed with a decent Slate widget.

diff --git a/Engine/Source/Editor/IntroTutorials/Private/IntroTutorials.cpp b/Engine/Source/Editor/IntroTutorials/Private/IntroTutorials.cpp
--- a/Engine/Source/Editor/IntroTutorials/Private/IntroTutorials.cpp
+++ b/Engine/Source/Editor/IntroTutorials/Private/IntroTutorials.cpp
@@ -27,6 +27,7 @@
 #include "EditorTutorialDetailsCustomization.h"
 #include "STutorialRoot.h"
 #include "STutorialButton.h"
+#include "STutorialLoading.h"
 #include "ToolkitManager.h"
 #include "BlueprintEditor.h"
 #include "EngineBuildSettings.h"
@@ -407,6 +408,22 @@ void FIntroTutorials::SummonTutorialBrowser()
 	}
 }
 
+void FIntroTutorials::AttachWidget(TSharedPtr<SWidget> Widget)
+{
+	if (TutorialRoot.IsValid())
+	{
+		TutorialRoot->AttachWidget(Widget);
+	}
+}
+
+void FIntroTutorials::DetachWidget()
+{
+	if (TutorialRoot.IsValid())
+	{
+		TutorialRoot->DetachWidget();
+	}
+}
+
 void FIntroTutorials::LaunchTutorial(const FString& TutorialAssetName)
 {
 	LaunchTutorialByName(TutorialAssetName);
@@ -466,6 +483,12 @@ TSharedRef<SWidget> FIntroTutorials::CreateTutorialsWidget(FName InContext, TWea
 		.ContextWindow(InContextWindow);
 }
 
+TSharedPtr<SWidget> FIntroTutorials::CreateTutorialsLoadingWidget(TWeakPtr<SWindow> InContextWindow) const
+{
+	return SNew(STutorialLoading)
+		.ContextWindow(InContextWindow);
+}
+
 float FIntroTutorials::GetIntroCurveValue(float InTime)
 {
 	if(ContentIntroCurve.IsValid())
diff --git a/Engine/Source/Editor/IntroTutorials/Private/IntroTutorials.h b/Engine/Source/Editor/IntroTutorials/Private/IntroTutorials.h
--- a/Engine/Source/Editor/IntroTutorials/Private/IntroTutorials.h
+++ b/Engine/Source/Editor/IntroTutorials/Private/IntroTutorials.h
@@ -34,6 +34,10 @@ public:
 
 	void SummonTutorialBrowser();
 
+	void AttachWidget(TSharedPtr<SWidget> Widget);
+
+	void DetachWidget();
+
 	static FString AnalyticsEventNameFromTutorial(UEditorTutorial* Tutorial);
 
 private:
@@ -98,6 +102,7 @@ public:
 	virtual void LaunchTutorial(UEditorTutorial* Tutorial, bool bRestart = true, TWeakPtr<SWindow> InNavigationWindow = nullptr, FSimpleDelegate OnTutorialClosed = FSimpleDelegate(), FSimpleDelegate OnTutorialExited = FSimpleDelegate()) override;
 	virtual void CloseAllTutorialContent() override;
 	virtual TSharedRef<SWidget> CreateTutorialsWidget(FName InContext, TWeakPtr<SWindow> InContextWindow = nullptr) const override;
+	virtual TSharedPtr<SWidget> CreateTutorialsLoadingWidget(TWeakPtr<SWindow> InContextWindow = nullptr) const override;
 	// End of IIntroTutorials interface
 private:
 	/** The tab id of the tutorial tab */
diff --git a/Engine/Source/Editor/IntroTutorials/Private/STutorialButton.cpp b/Engine/Source/Editor/IntroTutorials/Private/STutorialButton.cpp
--- a/Engine/Source/Editor/IntroTutorials/Private/STutorialButton.cpp
+++ b/Engine/Source/Editor/IntroTutorials/Private/STutorialButton.cpp
@@ -2,6 +2,7 @@
 
 #include "IntroTutorialsPrivatePCH.h"
 #include "STutorialButton.h"
+#include "STutorialLoading.h"
 #include "EditorTutorialSettings.h"
 #include "TutorialStateSettings.h"
 #include "TutorialMetaData.h"
@@ -34,6 +35,21 @@ void STutorialButton::Construct(const FArguments& InArgs)
 	PulseAnimation.AddCurve(0.0f, TutorialButtonConstants::PulseAnimationLength, ECurveEaseFunction::Linear);
 	RegisterActiveTimer( 0.f, FWidgetActiveTimerDelegate::CreateSP( this, &STutorialButton::OpenTutorialPostConstruct ) );
 
+	IIntroTutorials& IntroTutorials = FModuleManager::LoadModuleChecked<IIntroTutorials>(TEXT("IntroTutorials"));
+	LoadingWidget = IntroTutorials.CreateTutorialsLoadingWidget(ContextWindow);
+	/*
+	//LoadingWidget->SetVisibility(EVisibility::Visible);
+
+	if (GEngine)
+	{
+		UGameViewportClient* GVC = GEngine->Viewport;
+		if (GVC)
+		{
+			GVC->AddViewportWidgetContent(LoadingWidgetRef, 100);
+		}
+	}
+	*/
+
 	ChildSlot
 	[
 		SNew(SButton)
@@ -48,6 +64,8 @@ void STutorialButton::Construct(const FArguments& InArgs)
 			.HeightOverride(16)
 		]
 	];
+
+	//ChildSlot.AttachWidget(LoadingWidgetRef);
 }
 
 EActiveTimerReturnType STutorialButton::OpenTutorialPostConstruct( double InCurrentTime, float InDeltaTime )
@@ -159,6 +177,8 @@ FReply STutorialButton::HandleButtonClicked()
 
 	bPendingClickAction = true;
 	RegisterActiveTimer(0.f, FWidgetActiveTimerDelegate::CreateSP(this, &STutorialButton::HandleButtonClicked_AssetRegistryChecker));
+	FIntroTutorials& IntroTutorials = FModuleManager::GetModuleChecked<FIntroTutorials>(TEXT("IntroTutorials"));
+	IntroTutorials.AttachWidget(LoadingWidget);
 	return FReply::Handled();
 }
 
@@ -183,6 +203,7 @@ EActiveTimerReturnType STutorialButton::HandleButtonClicked_AssetRegistryChecker
 
 	//Now we know the asset registry is loaded, the tutorial broswer is updated, and we are ready to complete the click and stop this active timer
 	FIntroTutorials& IntroTutorials = FModuleManager::GetModuleChecked<FIntroTutorials>(TEXT("IntroTutorials"));
+	IntroTutorials.DetachWidget();
 	if (ShouldLaunchBrowser())
 	{
 		IntroTutorials.SummonTutorialBrowser();
diff --git a/Engine/Source/Editor/IntroTutorials/Private/STutorialButton.h b/Engine/Source/Editor/IntroTutorials/Private/STutorialButton.h
--- a/Engine/Source/Editor/IntroTutorials/Private/STutorialButton.h
+++ b/Engine/Source/Editor/IntroTutorials/Private/STutorialButton.h
@@ -2,6 +2,8 @@
 
 #pragma once
 
+#include "STutorialLoading.h"
+
 class STutorialButton : public SCompoundWidget
 {
 	SLATE_BEGIN_ARGS(STutorialButton) {}
@@ -90,6 +92,8 @@ private:
 	/** Cached browser filter */
 	FString CachedBrowserFilter;
 
+	TSharedPtr<SWidget> LoadingWidget;
+
 	/** True if we're waiting for asset registry to load in response to a click */
 	bool bPendingClickAction;
 };
\ No newline at end of file
diff --git a/Engine/Source/Editor/IntroTutorials/Private/STutorialLoading.cpp b/Engine/Source/Editor/IntroTutorials/Private/STutorialLoading.cpp
new file mode 100644
--- /dev/null
+++ b/Engine/Source/Editor/IntroTutorials/Private/STutorialLoading.cpp
@@ -0,0 +1,40 @@
+// Copyright 1998-2015 Epic Games, Inc. All Rights Reserved.
+
+#include "IntroTutorialsPrivatePCH.h"
+#include "STutorialLoading.h"
+#include "SThrobber.h"
+
+void STutorialLoading::Construct(const FArguments& InArgs)
+{
+	ContextWindow = InArgs._ContextWindow;
+
+	ChildSlot
+	.VAlign(VAlign_Center)
+	.HAlign(HAlign_Center)
+	[
+		SNew(SBorder)
+		.VAlign(VAlign_Center)
+		.HAlign(HAlign_Center)
+		.BorderImage(FEditorStyle::GetBrush("ToolPanel.GroupBorder"))
+		[
+			SNew(SVerticalBox)
+			+ SVerticalBox::Slot()
+			.VAlign(VAlign_Top)
+			.HAlign(HAlign_Center)
+			[
+				SNew(SCircularThrobber)
+			]
+			+ SVerticalBox::Slot()
+			.VAlign(VAlign_Top)
+			.HAlign(HAlign_Center)
+			[
+				SNew(STextBlock)
+				.ShadowColorAndOpacity(FLinearColor::Black)
+				.ColorAndOpacity(FLinearColor::White)
+				.ShadowOffset(FIntPoint(-1, 1))
+				.Font(FSlateFontInfo("Veranda", 16))
+				.Text(FText::FromString("Loading Tutorial Content"))
+			]
+		]
+	];
+}
diff --git a/Engine/Source/Editor/IntroTutorials/Private/STutorialLoading.h b/Engine/Source/Editor/IntroTutorials/Private/STutorialLoading.h
new file mode 100644
--- /dev/null
+++ b/Engine/Source/Editor/IntroTutorials/Private/STutorialLoading.h
@@ -0,0 +1,25 @@
+// Copyright 1998-2015 Epic Games, Inc. All Rights Reserved.
+
+#pragma once
+
+#include "SlateBasics.h"
+#include "SCompoundWidget.h"
+#include "SWindow.h"
+
+class STutorialLoading : public SCompoundWidget
+{
+	SLATE_BEGIN_ARGS(STutorialLoading) {}
+	
+	SLATE_ARGUMENT(TWeakPtr<SWindow>, ContextWindow)
+	
+	SLATE_END_ARGS()
+
+public:
+
+	/** Widget constructor */
+	void Construct(const FArguments& InArgs);
+
+private:
+	/** Window where loading visual will be displayed */
+	TWeakPtr<SWindow> ContextWindow;
+};
\ No newline at end of file
diff --git a/Engine/Source/Editor/IntroTutorials/Private/STutorialRoot.cpp b/Engine/Source/Editor/IntroTutorials/Private/STutorialRoot.cpp
--- a/Engine/Source/Editor/IntroTutorials/Private/STutorialRoot.cpp
+++ b/Engine/Source/Editor/IntroTutorials/Private/STutorialRoot.cpp
@@ -29,6 +29,18 @@ void STutorialRoot::Construct(const FArguments& InArgs)
 	];
 }
 
+void STutorialRoot::AttachWidget(TSharedPtr<SWidget> Widget)
+{
+	//This checkSlow is just here to ensure we know what we're doing
+	checkSlow(ChildSlot.GetWidget() == SNullWidget::NullWidget);
+	ChildSlot.AttachWidget(Widget.ToSharedRef());
+}
+
+void STutorialRoot::DetachWidget()
+{
+	ChildSlot.DetachWidget();
+}
+
 void STutorialRoot::MaybeAddOverlay(TSharedRef<SWindow> InWindow)
 {
 	if(InWindow->HasOverlay())
diff --git a/Engine/Source/Editor/IntroTutorials/Private/STutorialRoot.h b/Engine/Source/Editor/IntroTutorials/Private/STutorialRoot.h
--- a/Engine/Source/Editor/IntroTutorials/Private/STutorialRoot.h
+++ b/Engine/Source/Editor/IntroTutorials/Private/STutorialRoot.h
@@ -48,6 +48,12 @@ public:
 	/** Register that this widget was drawn this frame */
 	void WidgetWasDrawn(const FName& InName);
 
+	/** This is currently used for the "loading" widget */
+	void AttachWidget(TSharedPtr<SWidget> Widget);
+
+	/** This is currently used for the "loading" widget */
+	void DetachWidget();
+
 private:
 	/** Handle when the next button is clicked */
 	void HandleNextClicked(TWeakPtr<SWindow> InNavigationWindow);
diff --git a/Engine/Source/Editor/IntroTutorials/Public/IIntroTutorials.h b/Engine/Source/Editor/IntroTutorials/Public/IIntroTutorials.h
--- a/Engine/Source/Editor/IntroTutorials/Public/IIntroTutorials.h
+++ b/Engine/Source/Editor/IntroTutorials/Public/IIntroTutorials.h
@@ -66,5 +66,12 @@ public:
 	 * @return a widget used to access context-sensitive tutorials
 	 */
 	virtual TSharedRef<SWidget> CreateTutorialsWidget(FName InContext, TWeakPtr<SWindow> InContextWindow = nullptr) const = 0;
+
+	/**
+	* Create a widget that shows a "loading" visual while waiting for tutorials to load.
+	* @param	InContextWindow		The window that the context is attached to (e.g. the main window, or an asset editor tab)
+	* @return a widget used to access context-sensitive tutorials
+	*/
+	virtual TSharedPtr<SWidget> CreateTutorialsLoadingWidget(TWeakPtr<SWindow> InContextWindow = nullptr) const = 0;
 };
 
