UE-7933 - Added functionality to dismiss a summoned tutorial browser. This functionality is triggered when a tutorial is launched from the browser.

diff --git a/Engine/Source/Editor/IntroTutorials/Private/IntroTutorials.cpp b/Engine/Source/Editor/IntroTutorials/Private/IntroTutorials.cpp
--- a/Engine/Source/Editor/IntroTutorials/Private/IntroTutorials.cpp
+++ b/Engine/Source/Editor/IntroTutorials/Private/IntroTutorials.cpp
@@ -404,7 +404,16 @@ void FIntroTutorials::SummonTutorialBrowser()
 	if(TutorialRoot.IsValid())
 	{
 		FLevelEditorModule& LevelEditorModule = FModuleManager::GetModuleChecked<FLevelEditorModule>( TEXT("LevelEditor") );
-		LevelEditorModule.GetLevelEditorTabManager()->InvokeTab(FTabId("TutorialsBrowser"));
+		TutorialBrowserDockTab = LevelEditorModule.GetLevelEditorTabManager()->InvokeTab(FTabId("TutorialsBrowser"));
+	}
+}
+
+void FIntroTutorials::DismissTutorialBrowser()
+{
+	if (TutorialBrowserDockTab.IsValid() && TutorialBrowserDockTab.Pin().IsValid())
+	{
+		TutorialBrowserDockTab.Pin()->RequestCloseTab();
+		TutorialBrowserDockTab = nullptr;
 	}
 }
 
diff --git a/Engine/Source/Editor/IntroTutorials/Private/IntroTutorials.h b/Engine/Source/Editor/IntroTutorials/Private/IntroTutorials.h
--- a/Engine/Source/Editor/IntroTutorials/Private/IntroTutorials.h
+++ b/Engine/Source/Editor/IntroTutorials/Private/IntroTutorials.h
@@ -34,6 +34,8 @@ public:
 
 	void SummonTutorialBrowser();
 
+	void DismissTutorialBrowser();
+
 	void AttachWidget(TSharedPtr<SWidget> Widget);
 
 	void DetachWidget();
@@ -95,6 +97,9 @@ private:
 	/** Spawn a tab for browsing tutorials */
 	TSharedRef<SDockTab> SpawnTutorialsBrowserTab(const FSpawnTabArgs& SpawnTabArgs);
 
+	/** When we create a tutorial browser tab, we only keep a weak pointer to it. This way, it can be closed/destroyed without our involvement, but we can also close/destroy it ourselves. */
+	TWeakPtr<SDockTab> TutorialBrowserDockTab;
+
 public:
 
 	// IIntroTutorials interface
diff --git a/Engine/Source/Editor/IntroTutorials/Private/STutorialsBrowser.cpp b/Engine/Source/Editor/IntroTutorials/Private/STutorialsBrowser.cpp
--- a/Engine/Source/Editor/IntroTutorials/Private/STutorialsBrowser.cpp
+++ b/Engine/Source/Editor/IntroTutorials/Private/STutorialsBrowser.cpp
@@ -820,15 +820,23 @@ bool STutorialsBrowser::IsBackButtonEnabled() const
 
 void STutorialsBrowser::OnTutorialSelected(UEditorTutorial* InTutorial, bool bRestart)
 {
-	if( FEngineAnalytics::IsAvailable() && InTutorial != nullptr )
+	if (InTutorial != nullptr)
 	{
-		TArray<FAnalyticsEventAttribute> EventAttributes;
-		EventAttributes.Add(FAnalyticsEventAttribute(TEXT("Restarted"), bRestart));
-		EventAttributes.Add(FAnalyticsEventAttribute(TEXT("TutorialAsset"), FIntroTutorials::AnalyticsEventNameFromTutorial(InTutorial)));
+		if (FEngineAnalytics::IsAvailable())
+		{
+			TArray<FAnalyticsEventAttribute> EventAttributes;
+			EventAttributes.Add(FAnalyticsEventAttribute(TEXT("Restarted"), bRestart));
+			EventAttributes.Add(FAnalyticsEventAttribute(TEXT("TutorialAsset"), FIntroTutorials::AnalyticsEventNameFromTutorial(InTutorial)));
 
-		FEngineAnalytics::GetProvider().RecordEvent( TEXT("Rocket.Tutorials.LaunchedFromBrowser"), EventAttributes );
+			FEngineAnalytics::GetProvider().RecordEvent(TEXT("Rocket.Tutorials.LaunchedFromBrowser"), EventAttributes);
+		}
+		//Close the tutorial browser so it doesn't get in the way of the actual tutorial.
+		if (OnLaunchTutorial.IsBound())
+		{
+			FIntroTutorials& IntroTutorials = FModuleManager::GetModuleChecked<FIntroTutorials>(TEXT("IntroTutorials"));
+			IntroTutorials.DismissTutorialBrowser();
+		}
 	}
-
 	OnLaunchTutorial.ExecuteIfBound(InTutorial, bRestart, ParentWindow, FSimpleDelegate(), FSimpleDelegate());
 }
 
