[INTEGRATE] Change 2422514 by [REDACTED_PERSON][REDACTED_HANDLE] on 2015/01/28 13:45:27

	UE-8101 - CurrentTutorial wasn't being cleared when a tutorial ended, so some effects of tutorial stages could still happen after the user exited a tutorial.

diff --git a/Engine/Source/Editor/IntroTutorials/Private/STutorialRoot.cpp b/Engine/Source/Editor/IntroTutorials/Private/STutorialRoot.cpp
--- a/Engine/Source/Editor/IntroTutorials/Private/STutorialRoot.cpp
+++ b/Engine/Source/Editor/IntroTutorials/Private/STutorialRoot.cpp
@@ -347,6 +347,9 @@ void STutorialRoot::HandleCloseClicked()
 			
 		FEngineAnalytics::GetProvider().RecordEvent( TEXT("Rocket.Tutorials.Closed"), EventAttributes );
 	}
+
+	//Tutorial is no longer current
+	CurrentTutorial = nullptr;
 }
 
 bool STutorialRoot::WasWidgetDrawn(const FName& InName) const
