Sort order of tutorials and categories in tutorial browser can now be set in data instead of always being alphabetical.

diff --git a/Engine/Source/Editor/IntroTutorials/Private/STutorialsBrowser.cpp b/Engine/Source/Editor/IntroTutorials/Private/STutorialsBrowser.cpp
--- a/Engine/Source/Editor/IntroTutorials/Private/STutorialsBrowser.cpp
+++ b/Engine/Source/Editor/IntroTutorials/Private/STutorialsBrowser.cpp
@@ -166,6 +166,10 @@ public:
 		return Title;
 	}
 
+	int32 GetSortOrder() const override
+	{
+		return Category.SortOrder;
+	}
 
 	FText GetTitleText() const override
 	{
@@ -174,7 +178,7 @@ public:
 
 	bool SortAgainst(TSharedRef<ITutorialListEntry> OtherEntry) const override
 	{
-		return GetTitleString() < OtherEntry->GetTitleString();
+		return (GetSortOrder() == OtherEntry->GetSortOrder()) ? (GetTitleString() > OtherEntry->GetTitleString()) : (GetSortOrder() < OtherEntry->GetSortOrder());
 	}
 
 	void AddSubCategory(TSharedPtr<FTutorialListEntry_Category> InSubCategory)
@@ -388,9 +392,14 @@ public:
 		return Tutorial->Title.ToString();
 	}
 
+	int32 GetSortOrder() const override
+	{
+		return Tutorial->SortOrder;
+	}
+
 	bool SortAgainst(TSharedRef<ITutorialListEntry> OtherEntry) const override
 	{
-		return GetTitleString() < OtherEntry->GetTitleString();
+		return (GetSortOrder() == OtherEntry->GetSortOrder()) ? (GetTitleString() > OtherEntry->GetTitleString()) : (GetSortOrder() < OtherEntry->GetSortOrder());
 	}
 
 	FReply OnClicked(bool bRestart) const
diff --git a/Engine/Source/Editor/IntroTutorials/Private/STutorialsBrowser.h b/Engine/Source/Editor/IntroTutorials/Private/STutorialsBrowser.h
--- a/Engine/Source/Editor/IntroTutorials/Private/STutorialsBrowser.h
+++ b/Engine/Source/Editor/IntroTutorials/Private/STutorialsBrowser.h
@@ -21,6 +21,9 @@ struct ITutorialListEntry
 	/** Get the string representation of this item's title */
 	virtual FString GetTitleString() const = 0;
 
+	/** Get a priority value to override alphabetical sorting */
+	virtual int32 GetSortOrder() const = 0;
+
 	/** 
 	 * Sort this entry against another entry
 	 * @return true if this < OtherEntry 
diff --git a/Engine/Source/Editor/IntroTutorials/Public/EditorTutorial.h b/Engine/Source/Editor/IntroTutorials/Public/EditorTutorial.h
--- a/Engine/Source/Editor/IntroTutorials/Public/EditorTutorial.h
+++ b/Engine/Source/Editor/IntroTutorials/Public/EditorTutorial.h
@@ -55,6 +55,10 @@ struct INTROTUTORIALS_API FTutorialCategory
 	UPROPERTY(EditAnywhere, Category="Content")
 	FText Title;
 
+	/** Sort order, used by the tutorial browser - set in editor game-agnostic INI file */
+	UPROPERTY()
+	int32 SortOrder;
+
 	/** Localized text to use to describe this category */
 	UPROPERTY(EditAnywhere, Category="Content", meta=(MultiLine=true))
 	FText Description;
@@ -215,6 +219,10 @@ class INTROTUTORIALS_API UEditorTutorial : public UObject
 	UPROPERTY(EditAnywhere, Category="Tutorial", AssetRegistrySearchable)
 	FText Title;
 
+	/** Sorting priority, used by the tutorial browser */
+	UPROPERTY(BlueprintReadOnly, EditDefaultsOnly, Category = "Tutorial")
+	int32 SortOrder;
+
 	/** Icon name for this tutorial, used when presented to the user in the tutorial browser. This is a name for the icon in the Slate editor style. Only used if there isn't a valid texture to use. */
 	UPROPERTY(EditAnywhere, Category="Tutorial")
 	FString Icon;
