[INTEGRATE] Change 2427467 by bruce.nesbit@bnesbit_release_port on 2015/02/01 13:11:01

	Fix for assets in tutorials that have redirectors.

diff --git a/Engine/Source/Editor/IntroTutorials/Private/STutorialOverlay.cpp b/Engine/Source/Editor/IntroTutorials/Private/STutorialOverlay.cpp
--- a/Engine/Source/Editor/IntroTutorials/Private/STutorialOverlay.cpp
+++ b/Engine/Source/Editor/IntroTutorials/Private/STutorialOverlay.cpp
@@ -247,6 +247,7 @@ void STutorialOverlay::FocusOnAnyBlueprintNodes(const FTutorialWidgetContent &Wi
 	{
 		return;
 	}
+	
 	FString Name = WidgetContent.WidgetAnchor.OuterName;
 	int32 NameIndex;
 	Name.FindLastChar(TEXT('.'), NameIndex);
@@ -266,9 +267,30 @@ void STutorialOverlay::FocusOnAnyBlueprintNodes(const FTutorialWidgetContent &Wi
 	}
 	else /*if ( WidgetContent.WidgetAnchor.WrapperIdentifier.IsValid() == true )*/
 	{
+		const FName ObjectPath = WidgetContent.WidgetAnchor.WrapperIdentifier;
+
 		// if we didn't have a blueprint object to focus on, try it with a regular one
-		UObject* FocusObject = FindObject<UObject>(ANY_PACKAGE, *WidgetContent.WidgetAnchor.WrapperIdentifier.ToString());
-		if (FocusObject != nullptr)
+		UObject* FocusObject = FindObject<UObject>(ANY_PACKAGE, *ObjectPath.ToString());
+
+		// If we found an asset redirector, we need to follow it
+		UObjectRedirector* Redir = dynamic_cast<UObjectRedirector*>(FocusObject);
+		if (Redir)
+		{
+			FocusObject = Redir->DestinationObject;
+		}
+
+		// If we failed to find the object, it may be a class that has been redirected
+		if (!FocusObject)
+		{
+			const FString ObjectName = FPackageName::ObjectPathToObjectName(ObjectPath.ToString());
+			const FName RedirectedObjectName = ULinkerLoad::FindNewNameForClass(*ObjectName, false);
+			if (!RedirectedObjectName.IsNone())
+			{
+				FocusObject = FindObject<UClass>(ANY_PACKAGE, *RedirectedObjectName.ToString());
+			}
+		}
+
+		if (FocusObject)
 		{
 			TArray< UObject* > Objects;
 			Objects.Add(FocusObject);
