[INTEGRATE] Change 2431643 by bruce.nesbit@bnesbit_Releases on 2015/02/04 09:36:53

	Added code in STutorialOverlay to attempt to load an object if it isnt found when trying to focus.

diff --git a/Engine/Source/Editor/IntroTutorials/Private/STutorialOverlay.cpp b/Engine/Source/Editor/IntroTutorials/Private/STutorialOverlay.cpp
--- a/Engine/Source/Editor/IntroTutorials/Private/STutorialOverlay.cpp
+++ b/Engine/Source/Editor/IntroTutorials/Private/STutorialOverlay.cpp
@@ -248,7 +248,8 @@ void STutorialOverlay::FocusOnAnyBlueprintNodes(const FTutorialWidgetContent &Wi
 		return;
 	}
 	
-	FString Name = WidgetContent.WidgetAnchor.OuterName;
+	const FString Name = WidgetContent.WidgetAnchor.OuterName;
+	const FName ObjectPath = WidgetContent.WidgetAnchor.WrapperIdentifier;
 	int32 NameIndex;
 	Name.FindLastChar(TEXT('.'), NameIndex);
 	FString BlueprintName = Name.RightChop(NameIndex + 1);
@@ -265,13 +266,15 @@ void STutorialOverlay::FocusOnAnyBlueprintNodes(const FTutorialWidgetContent &Wi
 			FKismetEditorUtilities::BringKismetToFocusAttentionOnObject(GraphNode, false);
 		}
 	}
-	else /*if ( WidgetContent.WidgetAnchor.WrapperIdentifier.IsValid() == true )*/
+	else if ( !ObjectPath.IsNone() )
 	{
-		const FName ObjectPath = WidgetContent.WidgetAnchor.WrapperIdentifier;
-
 		// if we didn't have a blueprint object to focus on, try it with a regular one
 		UObject* FocusObject = FindObject<UObject>(ANY_PACKAGE, *ObjectPath.ToString());
-
+		// If we didn't find it, maybe it just hasn't been loaded yet
+		if( FocusObject == nullptr )
+		{
+			FocusObject = LoadObject<UObject>(nullptr, *ObjectPath.ToString(),nullptr, LOAD_FindIfFail);
+		}
 		// If we found an asset redirector, we need to follow it
 		UObjectRedirector* Redir = dynamic_cast<UObjectRedirector*>(FocusObject);
 		if (Redir)
