[INTEGRATE] Change 2441350 by bruce.nesbit@bnesbit_Releases on 2015/02/11 05:55:27

	Added back button to tutorials.
	(Also fixed overridable text on next button and provided same for back button)

diff --git a/Engine/Source/Editor/IntroTutorials/Private/STutorialContent.cpp b/Engine/Source/Editor/IntroTutorials/Private/STutorialContent.cpp
--- a/Engine/Source/Editor/IntroTutorials/Private/STutorialContent.cpp
+++ b/Engine/Source/Editor/IntroTutorials/Private/STutorialContent.cpp
@@ -45,6 +45,7 @@ void STutorialContent::Construct(const FArguments& InArgs, UEditorTutorial* InTu
 	bAllowNonWidgetContent = InArgs._AllowNonWidgetContent;
 	OnWasWidgetDrawn = InArgs._OnWasWidgetDrawn;
 	NextButtonText = InArgs._NextButtonText;
+	BackButtonText = InArgs._BackButtonText;
 
 	BorderIntroAnimation.AddCurve(0.0f, TutorialConstants::BorderIntroAnimationLength, ECurveEaseFunction::CubicOut);
 	BorderPulseAnimation.AddCurve(0.0f, TutorialConstants::BorderPulseAnimationLength, ECurveEaseFunction::Linear);
@@ -143,7 +144,47 @@ void STutorialContent::Construct(const FArguments& InArgs, UEditorTutorial* InTu
 					]
 				]
 			]
-			+SOverlay::Slot()
+			+ SOverlay::Slot()
+			.VAlign(VAlign_Bottom)
+			.HAlign(HAlign_Left)
+			.Padding(12.0f)
+			[
+				SAssignNew(BackButton, SButton)
+				.ToolTipText(this, &STutorialContent::GetBackButtonTooltip)
+				.OnClicked(this, &STutorialContent::HandleBackButtonClicked)
+				.Visibility(this, &STutorialContent::GetBackButtonVisibility)
+				.ButtonStyle(&FEditorStyle::Get().GetWidgetStyle<FButtonStyle>("Tutorials.Content.NavigationButtonWrapper"))
+				.ContentPadding(0.0f)
+				[
+					SNew(SBox)
+					.Padding(8.0f)
+					[
+						SNew(SBorder)
+						.BorderImage(this, &STutorialContent::GetBackButtonBorder)
+						[
+							SNew(SHorizontalBox)
+							+ SHorizontalBox::Slot()
+							.AutoWidth()
+							.VAlign(VAlign_Center)
+							.Padding(4.0f, 0.0f, 0.0f, 0.0f)
+							[
+								SNew(STextBlock)
+								.Text(this, &STutorialContent::GetBackButtonLabel)
+								.TextStyle(FEditorStyle::Get(), "Tutorials.Content.NavigationText")
+								.ColorAndOpacity(FLinearColor::White)
+							]
+							+ SHorizontalBox::Slot()
+								.AutoWidth()
+								[
+									SNew(SImage)
+									.Image(this, &STutorialContent::GetBackButtonBrush)
+									.ColorAndOpacity(FLinearColor::White)
+								]
+						]
+					]
+				]
+			] 
+			+ SOverlay::Slot()
 			.VAlign(VAlign_Bottom)
 			.HAlign(HAlign_Right)
 			.Padding(12.0f)
@@ -684,6 +725,16 @@ FReply STutorialContent::HandleNextClicked()
 	return FReply::Handled();
 }
 
+FReply STutorialContent::HandleBackButtonClicked()
+{
+	if (IsBackEnabled.Get())
+	{
+		OnBackClicked.ExecuteIfBound();
+	}
+
+	return FReply::Handled();
+}
+
 const FSlateBrush* STutorialContent::GetNextButtonBrush() const
 {
 	if(IsNextEnabled.Get())
@@ -732,6 +783,51 @@ const FSlateBrush* STutorialContent::GetNextButtonBorder() const
 	return NextButton->IsHovered() ? &FEditorStyle::Get().GetWidgetStyle<FButtonStyle>("Tutorials.Content.NavigationButton").Hovered : &FEditorStyle::Get().GetWidgetStyle<FButtonStyle>("Tutorials.Content.NavigationButton").Normal;
 }
 
+const FSlateBrush* STutorialContent::GetBackButtonBrush() const
+{
+	if (IsBackEnabled.Get())
+	{
+		return FEditorStyle::GetBrush("Tutorials.Navigation.BackButton");
+	}
+	return FEditorStyle::GetDefaultBrush();
+}
+
+EVisibility STutorialContent::GetBackButtonVisibility() const
+{
+	return IsBackEnabled.Get() == true ? EVisibility::Visible : EVisibility::Collapsed;
+}
+
+FText STutorialContent::GetBackButtonTooltip() const
+{
+	if (IsBackEnabled.Get())
+	{
+		return LOCTEXT("BackButtonTooltip", "Go to the previous stage of this tutorial.");
+	}
+	return FText::GetEmpty();
+}
+
+FText STutorialContent::GetBackButtonLabel() const
+{
+	if (!BackButtonText.Get().IsEmpty())
+	{
+		return BackButtonText.Get();
+	}
+	else
+	{
+		if (IsBackEnabled.Get())
+		{
+			return LOCTEXT("DefaultBackButtonLabel", "Back");
+		}	
+		
+	}
+	return FText::GetEmpty();
+}
+
+const FSlateBrush* STutorialContent::GetBackButtonBorder() const
+{
+	return BackButton->IsHovered() ? &FEditorStyle::Get().GetWidgetStyle<FButtonStyle>("Tutorials.Content.NavigationButton").Hovered : &FEditorStyle::Get().GetWidgetStyle<FButtonStyle>("Tutorials.Content.NavigationButton").Normal;
+}
+
 FReply STutorialContent::OnMouseButtonDown( const FGeometry& MyGeometry, const FPointerEvent& MouseEvent )
 {
 	// Mouse back and forward buttons traverse history
diff --git a/Engine/Source/Editor/IntroTutorials/Private/STutorialContent.h b/Engine/Source/Editor/IntroTutorials/Private/STutorialContent.h
--- a/Engine/Source/Editor/IntroTutorials/Private/STutorialContent.h
+++ b/Engine/Source/Editor/IntroTutorials/Private/STutorialContent.h
@@ -66,6 +66,9 @@ class STutorialContent : public SCompoundWidget
 	/** Text to display on next/home button */
 	SLATE_ATTRIBUTE(FText, NextButtonText)
 
+	/** Text to display on back button */
+	SLATE_ATTRIBUTE(FText, BackButtonText)
+
 	SLATE_END_ARGS()
 
 	void Construct(const FArguments& InArgs, UEditorTutorial* InTutorial, const FTutorialContent& InContent);
@@ -123,10 +126,10 @@ private:
 	/** Delegate handler for exiting the tutorial */
 	void HandleExitSelected();
 
-	/** Delegate handler for going to the previous stage of the tutorial */
+	/** Delegate handler for going to the previous stage of the tutorial (From dropdown menu) */
 	void HandleBackSelected();
-
-	/** Delegate handler for going to the next stage of the tutorial */
+	
+	/** Delegate handler for going to the next stage of the tutorial (From dropdown menu) */
 	void HandleNextSelected();
 
 	/** Delegate handler for restarting the tutorial */
@@ -135,9 +138,12 @@ private:
 	/** Delegate handler for exiting the tutorial to the browser */
 	void HandleBrowseSelected();
 
-	/** Delegate handler for going to the next stage of the tutorial */
+	/** Delegate handler for going to the next stage of the tutorial (From button) */
 	FReply HandleNextClicked();
 
+	/** Delegate handler for going to the previous stage of the tutorial (From button) */
+	FReply HandleBackButtonClicked();
+
 	/** Delegate handler allowing us to change the brush of the 'next' button depending on context */
 	const FSlateBrush* GetNextButtonBrush() const;
 
@@ -153,6 +159,21 @@ private:
 	/** Helper to determine the proper animation values for the border pulse */
 	void GetAnimationValues(float& OutAlphaFactor, float& OutPulseFactor, FLinearColor& OutShadowTint, FLinearColor& OutBorderTint) const;
 
+	/** Delegate handler allowing us to change the brush of the 'back' button depending on context */
+	const FSlateBrush* GetBackButtonBrush() const;
+
+	/** Delegate handler allowing us to change the tootlip of the 'back' button depending on context */
+	FText GetBackButtonTooltip() const;
+
+	/** Change back button color based on hover state */
+	FText GetBackButtonLabel() const;
+
+	/** We need to override the border ourselves, rather than let the button handle it, as we are using a larger apparent hitbox */
+	const FSlateBrush* GetBackButtonBorder() const;
+
+	/* Get the visibilty of the back button */
+	EVisibility GetBackButtonVisibility() const;
+
 private:
 
 	/** Copy of the window size we were last draw at */
@@ -222,6 +243,9 @@ private:
 	/** Next button widget */
 	TSharedPtr<SWidget> NextButton;
 
+	/** Back button widget */
+	TSharedPtr<SWidget> BackButton;
+
 	/** Whether we can show full window content in this overlay (i.e. in the same window as the navigation controls) */
 	bool bAllowNonWidgetContent;
 
@@ -230,4 +254,7 @@ private:
 
 	/** Text for next/home button */
 	TAttribute<FText> NextButtonText;
+
+	/** Text for next/home button */
+	TAttribute<FText> BackButtonText;
 };
\ No newline at end of file
diff --git a/Engine/Source/Editor/IntroTutorials/Private/STutorialOverlay.cpp b/Engine/Source/Editor/IntroTutorials/Private/STutorialOverlay.cpp
--- a/Engine/Source/Editor/IntroTutorials/Private/STutorialOverlay.cpp
+++ b/Engine/Source/Editor/IntroTutorials/Private/STutorialOverlay.cpp
@@ -54,6 +54,8 @@ void STutorialOverlay::Construct(const FArguments& InArgs, UEditorTutorial* InTu
 					.IsNextEnabled(InArgs._IsNextEnabled)
 					.IsStandalone(InArgs._IsStandalone)
 					.WrapTextAt(600.0f)
+					.NextButtonText(InStage->NextButtonText)
+					.BackButtonText(InStage->BackButtonText)
 				]
 			];
 		}
@@ -83,8 +85,7 @@ void STutorialOverlay::Construct(const FArguments& InArgs, UEditorTutorial* InTu
 						.WrapTextAt(WidgetContent.ContentWidth)
 						.Anchor(WidgetContent.WidgetAnchor)
 						.AllowNonWidgetContent(InArgs._AllowNonWidgetContent)
-						.OnWasWidgetDrawn(InArgs._OnWasWidgetDrawn)
-						.NextButtonText(InStage->NextButtonText);
+						.OnWasWidgetDrawn(InArgs._OnWasWidgetDrawn);
 					
 					PerformWidgetInteractions(InTutorial, WidgetContent); 					
 
diff --git a/Engine/Source/Editor/IntroTutorials/Public/EditorTutorial.h b/Engine/Source/Editor/IntroTutorials/Public/EditorTutorial.h
--- a/Engine/Source/Editor/IntroTutorials/Public/EditorTutorial.h
+++ b/Engine/Source/Editor/IntroTutorials/Public/EditorTutorial.h
@@ -207,6 +207,10 @@ struct INTROTUTORIALS_API FTutorialStage
 	/** Text to display on the next button */
 	UPROPERTY(EditAnywhere, Category="Stage")
 	FText NextButtonText;
+	
+	/** Text to display on the back button */
+	UPROPERTY(EditAnywhere, Category = "Stage")
+	FText BackButtonText;	
 };
 
 /** An asset used to build a stage-by-stage tutorial in the editor */
