[INTEGRATE] Change 2442585 by [REDACTED_PERSON][REDACTED_HANDLE] on 2015/02/11 18:24:57

	Always remind users about tutorials until they complete the tutorial or explicitly dismiss it (UE-9550)
	- Previously, even opening a tutorial once would prevent the reminder beacon for that tutorial to show again, ever.
	- Now, the beacon will be hidden for the remainder of that editor session and reappear the next time, unless the user completed the tutorial
	- The user can also right click on the beacon and choose to "Don't Remind Me Again" which will prevent it from appearing in the future for that specific tutorial (no change in behavior)

diff --git a/Engine/Source/Editor/IntroTutorials/Private/STutorialButton.cpp b/Engine/Source/Editor/IntroTutorials/Private/STutorialButton.cpp
--- a/Engine/Source/Editor/IntroTutorials/Private/STutorialButton.cpp
+++ b/Engine/Source/Editor/IntroTutorials/Private/STutorialButton.cpp
@@ -216,7 +216,9 @@ EActiveTimerReturnType STutorialButton::HandleButtonClicked_AssetRegistryChecker
 		const bool bRestart = true;
 		IntroTutorials.LaunchTutorial(CachedLaunchTutorial, bRestart, ContextWindow, Delegate, Delegate);
 
-		const bool bDismissAcrossSessions = true;
+		// The user asked to start the tutorial, so we don't need to remind them about it again this session, but we'll remind
+		// them in the next session if they haven't completed it by then
+		const bool bDismissAcrossSessions = false;
 		GetMutableDefault<UTutorialStateSettings>()->DismissTutorial(CachedLaunchTutorial, bDismissAcrossSessions);
 		GetMutableDefault<UTutorialStateSettings>()->SaveProgress();
 		bTutorialDismissed = true;
@@ -236,11 +238,13 @@ FReply STutorialButton::OnMouseButtonDown(const FGeometry& MyGeometry, const FPo
 		if(ShouldShowAlert())
 		{
 			MenuBuilder.AddMenuEntry(
-				LOCTEXT("DismissReminder", "Dismiss Alert"),
-				LOCTEXT("DismissReminderTooltip", "Don't show me this alert again"),
+				LOCTEXT("DismissReminder", "Don't Remind Me Again"),
+				LOCTEXT("DismissReminderTooltip", "Selection this option will prevent the tutorial blip from being displayed again, even if you choose not to complete the tutorial."),
 				FSlateIcon(),
 				FUIAction(FExecuteAction::CreateSP(this, &STutorialButton::DismissAlert))
 				);
+
+			MenuBuilder.AddMenuSeparator();
 		}
 
 		if(bTutorialAvailable)
@@ -279,6 +283,7 @@ void STutorialButton::DismissAlert()
 		FEngineAnalytics::GetProvider().RecordEvent( TEXT("Rocket.Tutorials.DismissedTutorialAlert"), EventAttributes );
 	}
 
+	// If they actually right click and choose "Dismiss Alert", we'll go ahead and suppress the tutorial reminder for this feature for good (all sessions.)
 	const bool bDismissAcrossSessions = true;
 	if (CachedAttractTutorial != nullptr)
 	{
@@ -322,7 +327,7 @@ FText STutorialButton::GetButtonToolTip() const
 {
 	if(ShouldLaunchBrowser())
 	{
-		return LOCTEXT("TutorialLaunchBrowserToolTip", "Show Available Tutorials");
+		return LOCTEXT("TutorialLaunchBrowserToolTip", "Show Available Tutorials...");
 	}
 	else if(bTutorialAvailable)
 	{
