[UE-9935] Moving fix from 4.7 back into Main.

diff --git a/Engine/Source/Editor/IntroTutorials/Private/EditorTutorialSettings.cpp b/Engine/Source/Editor/IntroTutorials/Private/EditorTutorialSettings.cpp
--- a/Engine/Source/Editor/IntroTutorials/Private/EditorTutorialSettings.cpp
+++ b/Engine/Source/Editor/IntroTutorials/Private/EditorTutorialSettings.cpp
@@ -17,12 +17,14 @@ void UEditorTutorialSettings::FindTutorialInfoForContext(FName InContext, UEdito
 		{
 			OutBrowserFilter = Context.BrowserFilter;
 
+			Context.LaunchTutorial.TryLoad();
 			TSubclassOf<UEditorTutorial> LaunchTutorialClass = Context.LaunchTutorial.ResolveClass();
 			if (LaunchTutorialClass != nullptr)
 			{
 				OutLaunchTutorial = LaunchTutorialClass->GetDefaultObject<UEditorTutorial>();
 			}
 
+			Context.AttractTutorial.TryLoad();
 			TSubclassOf<UEditorTutorial> AttractTutorialClass = Context.AttractTutorial.ResolveClass();
 			if (AttractTutorialClass != nullptr)
 			{
