[INTEGRATE] Change 2443386 by [REDACTED_PERSON][REDACTED_HANDLE] on 2015/02/12 09:36:51

	Fixed typo in tutorial button tool-tip change

diff --git a/Engine/Source/Editor/IntroTutorials/Private/STutorialButton.cpp b/Engine/Source/Editor/IntroTutorials/Private/STutorialButton.cpp
--- a/Engine/Source/Editor/IntroTutorials/Private/STutorialButton.cpp
+++ b/Engine/Source/Editor/IntroTutorials/Private/STutorialButton.cpp
@@ -239,7 +239,7 @@ FReply STutorialButton::OnMouseButtonDown(const FGeometry& MyGeometry, const FPo
 		{
 			MenuBuilder.AddMenuEntry(
 				LOCTEXT("DismissReminder", "Don't Remind Me Again"),
-				LOCTEXT("DismissReminderTooltip", "Selection this option will prevent the tutorial blip from being displayed again, even if you choose not to complete the tutorial."),
+				LOCTEXT("DismissReminderTooltip", "Selecting this option will prevent the tutorial blip from being displayed again, even if you choose not to complete the tutorial."),
 				FSlateIcon(),
 				FUIAction(FExecuteAction::CreateSP(this, &STutorialButton::DismissAlert))
 				);
