Font fix in tutorial loading widget. Minor cleanup (deleted commented-out code) in tutorial button.

diff --git a/Engine/Source/Editor/IntroTutorials/Private/STutorialButton.cpp b/Engine/Source/Editor/IntroTutorials/Private/STutorialButton.cpp
--- a/Engine/Source/Editor/IntroTutorials/Private/STutorialButton.cpp
+++ b/Engine/Source/Editor/IntroTutorials/Private/STutorialButton.cpp
@@ -37,18 +37,6 @@ void STutorialButton::Construct(const FArguments& InArgs)
 
 	IIntroTutorials& IntroTutorials = FModuleManager::LoadModuleChecked<IIntroTutorials>(TEXT("IntroTutorials"));
 	LoadingWidget = IntroTutorials.CreateTutorialsLoadingWidget(ContextWindow);
-	/*
-	//LoadingWidget->SetVisibility(EVisibility::Visible);
-
-	if (GEngine)
-	{
-		UGameViewportClient* GVC = GEngine->Viewport;
-		if (GVC)
-		{
-			GVC->AddViewportWidgetContent(LoadingWidgetRef, 100);
-		}
-	}
-	*/
 
 	ChildSlot
 	[
@@ -64,8 +52,6 @@ void STutorialButton::Construct(const FArguments& InArgs)
 			.HeightOverride(16)
 		]
 	];
-
-	//ChildSlot.AttachWidget(LoadingWidgetRef);
 }
 
 EActiveTimerReturnType STutorialButton::OpenTutorialPostConstruct( double InCurrentTime, float InDeltaTime )
diff --git a/Engine/Source/Editor/IntroTutorials/Private/STutorialLoading.cpp b/Engine/Source/Editor/IntroTutorials/Private/STutorialLoading.cpp
--- a/Engine/Source/Editor/IntroTutorials/Private/STutorialLoading.cpp
+++ b/Engine/Source/Editor/IntroTutorials/Private/STutorialLoading.cpp
@@ -4,37 +4,41 @@
 #include "STutorialLoading.h"
 #include "SThrobber.h"
 
+#define LOCTEXT_NAMESPACE "Tutorials"
+
 void STutorialLoading::Construct(const FArguments& InArgs)
 {
 	ContextWindow = InArgs._ContextWindow;
 
 	ChildSlot
-	.VAlign(VAlign_Center)
-	.HAlign(HAlign_Center)
-	[
-		SNew(SBorder)
 		.VAlign(VAlign_Center)
 		.HAlign(HAlign_Center)
-		.BorderImage(FEditorStyle::GetBrush("ToolPanel.GroupBorder"))
 		[
-			SNew(SVerticalBox)
-			+ SVerticalBox::Slot()
-			.VAlign(VAlign_Top)
+			SNew(SBorder)
+			.VAlign(VAlign_Center)
 			.HAlign(HAlign_Center)
+			.BorderImage(FEditorStyle::GetBrush("ToolPanel.GroupBorder"))
 			[
-				SNew(SCircularThrobber)
+				SNew(SVerticalBox)
+				+ SVerticalBox::Slot()
+				.VAlign(VAlign_Top)
+				.HAlign(HAlign_Center)
+				[
+					SNew(SCircularThrobber)
+				]
+				+ SVerticalBox::Slot()
+					.VAlign(VAlign_Top)
+					.HAlign(HAlign_Center)
+					[
+						SNew(STextBlock)
+						.ShadowColorAndOpacity(FLinearColor::Black)
+						.ColorAndOpacity(FLinearColor::White)
+						.ShadowOffset(FIntPoint(-1, 1))
+						.Font(FSlateFontInfo(FPaths::EngineContentDir() / TEXT("Slate/Fonts/Roboto-Regular.ttf"), 16))
+						.Text(LOCTEXT("LoadingContentTut", "Loading Tutorial Content"))
+					]
 			]
-			+ SVerticalBox::Slot()
-			.VAlign(VAlign_Top)
-			.HAlign(HAlign_Center)
-			[
-				SNew(STextBlock)
-				.ShadowColorAndOpacity(FLinearColor::Black)
-				.ColorAndOpacity(FLinearColor::White)
-				.ShadowOffset(FIntPoint(-1, 1))
-				.Font(FSlateFontInfo("Veranda", 16))
-				.Text(FText::FromString("Loading Tutorial Content"))
-			]
-		]
-	];
+		];
 }
+
+#undef LOCTEXT_NAMESPACE
