Added show/hide engine content blueprint callable function

diff --git a/Engine/Source/Editor/IntroTutorials/Private/EditorTutorial.cpp b/Engine/Source/Editor/IntroTutorials/Private/EditorTutorial.cpp
--- a/Engine/Source/Editor/IntroTutorials/Private/EditorTutorial.cpp
+++ b/Engine/Source/Editor/IntroTutorials/Private/EditorTutorial.cpp
@@ -80,4 +80,29 @@ AActor* UEditorTutorial::GetActorReference(FString PathToActor)
 #else
 	return nullptr;
 #endif //WITH_EDITOR
-}
\ No newline at end of file
+}
+
+void UEditorTutorial::SetEngineFolderVisibilty(bool bNewVisibility)
+{
+	bool bDisplayEngine = GetDefault<UContentBrowserSettings>()->GetDisplayEngineFolder();
+	// If we cannot change the vis state, or it matches the new request leave it
+	if (bDisplayEngine == bNewVisibility)
+	{
+		return;
+	}
+	if (bNewVisibility)
+	{
+		GetMutableDefault<UContentBrowserSettings>()->SetDisplayEngineFolder(true);
+	}
+	else
+	{
+		GetMutableDefault<UContentBrowserSettings>()->SetDisplayEngineFolder(false);
+		GetMutableDefault<UContentBrowserSettings>()->SetDisplayEngineFolder(false, true);
+	}
+	GetMutableDefault<UContentBrowserSettings>()->PostEditChange();
+}
+
+bool UEditorTutorial::GetEngineFolderVisibilty()
+{
+	return GetDefault<UContentBrowserSettings>()->GetDisplayEngineFolder();
+}
diff --git a/Engine/Source/Editor/IntroTutorials/Public/EditorTutorial.h b/Engine/Source/Editor/IntroTutorials/Public/EditorTutorial.h
--- a/Engine/Source/Editor/IntroTutorials/Public/EditorTutorial.h
+++ b/Engine/Source/Editor/IntroTutorials/Public/EditorTutorial.h
@@ -329,4 +329,16 @@ protected:
 	 */
 	UFUNCTION(BlueprintCallable, Category="Tutorial")
 	static void OpenAsset(UObject* Asset);
+
+	/**
+	 * Sets the visibility of the engine folder in the content browser
+	 */
+	UFUNCTION(BlueprintCallable, Category="Tutorial")
+	static void SetEngineFolderVisibilty(bool bNewVisibility);
+
+	/**
+	 * Returns the visibility of the engine folder in the content browser
+	 */
+	UFUNCTION(BlueprintCallable, Category="Tutorial")
+	static bool GetEngineFolderVisibilty();
 };
\ No newline at end of file
