[UE-11086] Dismissing a tutorial doesn't open the tutorial browser.

diff --git a/Engine/Source/Editor/IntroTutorials/Private/STutorialButton.cpp b/Engine/Source/Editor/IntroTutorials/Private/STutorialButton.cpp
--- a/Engine/Source/Editor/IntroTutorials/Private/STutorialButton.cpp
+++ b/Engine/Source/Editor/IntroTutorials/Private/STutorialButton.cpp
@@ -274,16 +274,17 @@ void STutorialButton::DismissAlert()
 	if (CachedAttractTutorial != nullptr)
 	{
 		GetMutableDefault<UTutorialStateSettings>()->DismissTutorial(CachedAttractTutorial, bDismissAcrossSessions);
+		CachedAttractTutorial = nullptr;
 	}
 	if( CachedLaunchTutorial != nullptr)
 	{
 		GetMutableDefault<UTutorialStateSettings>()->DismissTutorial(CachedLaunchTutorial, bDismissAcrossSessions);
+		CachedLaunchTutorial = nullptr;
 	}
 	GetMutableDefault<UTutorialStateSettings>()->SaveProgress();
 	bTutorialDismissed = true;
 
-	FIntroTutorials& IntroTutorials = FModuleManager::GetModuleChecked<FIntroTutorials>(TEXT("IntroTutorials"));
-	IntroTutorials.CloseAllTutorialContent();
+	RefreshStatus();
 }
 
 void STutorialButton::LaunchTutorial()
