Fixed tutorial analytics when returning to browser via 'Home' button

This needs to use a seperate path now so we cant rely on the 'closed' event firing.

AN-93 - Tutorial Event Not Sending?

diff --git a/Engine/Source/Editor/IntroTutorials/Private/STutorialRoot.cpp b/Engine/Source/Editor/IntroTutorials/Private/STutorialRoot.cpp
--- a/Engine/Source/Editor/IntroTutorials/Private/STutorialRoot.cpp
+++ b/Engine/Source/Editor/IntroTutorials/Private/STutorialRoot.cpp
@@ -217,6 +217,22 @@ void STutorialRoot::HandleHomeClicked()
 		GetMutableDefault<UTutorialStateSettings>()->SaveProgress();
 	}
 
+	// submit analytics data
+	if( FEngineAnalytics::IsAvailable() && CurrentTutorial != nullptr && CurrentTutorialStage < CurrentTutorial->Stages.Num() )
+	{
+		FString const CurrentExcerptTitle = CurrentTutorial->Stages[CurrentTutorialStage].Name.ToString();
+		int32 const CurrentExcerptIndex = CurrentTutorialStage;
+		float const CurrentPageElapsedTime = (float)(FPlatformTime::Seconds() - CurrentTutorialStartTime);
+
+		TArray<FAnalyticsEventAttribute> EventAttributes;
+		EventAttributes.Add(FAnalyticsEventAttribute(TEXT("LastStageIndex"), CurrentExcerptIndex));
+		EventAttributes.Add(FAnalyticsEventAttribute(TEXT("LastStageTitle"), CurrentExcerptTitle));
+		EventAttributes.Add(FAnalyticsEventAttribute(TEXT("TimeSpentInTutorial"), CurrentPageElapsedTime));
+		EventAttributes.Add(FAnalyticsEventAttribute(TEXT("TutorialAsset"), FIntroTutorials::AnalyticsEventNameFromTutorial(CurrentTutorial)));
+			
+		FEngineAnalytics::GetProvider().RecordEvent( TEXT("Rocket.Tutorials.Home"), EventAttributes );
+	}
+
 	CurrentTutorial = nullptr;
 	CurrentTutorialStage = 0;
 
