Fixed docked tutorial browser never showing tutorials

While the asset registry is loading we need to refresh tutorials as this is the only mechanism by which they are picked up by the browser.

diff --git a/Engine/Source/Editor/IntroTutorials/Private/STutorialsBrowser.cpp b/Engine/Source/Editor/IntroTutorials/Private/STutorialsBrowser.cpp
--- a/Engine/Source/Editor/IntroTutorials/Private/STutorialsBrowser.cpp
+++ b/Engine/Source/Editor/IntroTutorials/Private/STutorialsBrowser.cpp
@@ -697,12 +697,6 @@ void STutorialsBrowser::RebuildTutorials(TSharedPtr<FTutorialListEntry_Category>
 
 	//Ensure that tutorials are loaded into the asset registry before making a list of them.
 	FAssetRegistryModule& AssetRegistry = FModuleManager::LoadModuleChecked<FAssetRegistryModule>(TEXT("AssetRegistry"));
-	bool IsStillLoading = AssetRegistry.Get().IsLoadingAssets();
-	if (IsStillLoading)
-	{
-		//This can happen if you close the editor with the tutorials browser open, then restart the editor so that the browser opens immediately.
-		return;
-	}
 
 	// rebuild tutorials
 	FARFilter Filter;
