Enabled UHT to digest GENERATED_BODY instead of GENERATED_UCLASS_BODY, GENERATED_USTRUCT_BODY, GENERATED_UINTERFACE_BODY or GENERATED_IINTERFACE_BODY, changed every occurence to the new syntax and fix

diff --git a/Engine/Source/Editor/BlueprintGraph/Classes/K2Node_DelegateSet.h b/Engine/Source/Editor/BlueprintGraph/Classes/K2Node_DelegateSet.h
--- a/Engine/Source/Editor/BlueprintGraph/Classes/K2Node_DelegateSet.h
+++ b/Engine/Source/Editor/BlueprintGraph/Classes/K2Node_DelegateSet.h
@@ -9,7 +9,9 @@
 UCLASS(MinimalAPI)
 class UK2Node_DelegateSet : public UK2Node
 {
-	GENERATED_UCLASS_BODY()
+	GENERATED_BODY()
+public:
+	BLUEPRINTGRAPH_API UK2Node_DelegateSet(const FObjectInitializer& ObjectInitializer = FObjectInitializer::Get());
 
 	/** Delegate property name that this event is associated with on the target */
 	UPROPERTY()
diff --git a/Engine/Source/Editor/BlueprintGraph/Classes/K2Node_DoOnceMultiInput.h b/Engine/Source/Editor/BlueprintGraph/Classes/K2Node_DoOnceMultiInput.h
--- a/Engine/Source/Editor/BlueprintGraph/Classes/K2Node_DoOnceMultiInput.h
+++ b/Engine/Source/Editor/BlueprintGraph/Classes/K2Node_DoOnceMultiInput.h
@@ -7,7 +7,9 @@
 UCLASS(MinimalAPI)
 class UK2Node_DoOnceMultiInput : public UK2Node
 {
-	GENERATED_UCLASS_BODY()
+	GENERATED_BODY()
+public:
+	BLUEPRINTGRAPH_API UK2Node_DoOnceMultiInput(const FObjectInitializer& ObjectInitializer = FObjectInitializer::Get());
 
 	/** The number of additional input pins to generate for this node (2 base pins are not included) */
 	UPROPERTY()
diff --git a/Engine/Source/Editor/BlueprintGraph/Classes/K2Node_DynamicCast.h b/Engine/Source/Editor/BlueprintGraph/Classes/K2Node_DynamicCast.h
--- a/Engine/Source/Editor/BlueprintGraph/Classes/K2Node_DynamicCast.h
+++ b/Engine/Source/Editor/BlueprintGraph/Classes/K2Node_DynamicCast.h
@@ -9,7 +9,9 @@
 UCLASS(MinimalAPI)
 class UK2Node_DynamicCast : public UK2Node
 {
-	GENERATED_UCLASS_BODY()
+	GENERATED_BODY()
+public:
+	BLUEPRINTGRAPH_API UK2Node_DynamicCast(const FObjectInitializer& ObjectInitializer = FObjectInitializer::Get());
 
 	/** The type that the input should try to be cast to */
 	UPROPERTY()
diff --git a/Engine/Source/Editor/BlueprintGraph/Classes/K2Node_EaseFunction.h b/Engine/Source/Editor/BlueprintGraph/Classes/K2Node_EaseFunction.h
--- a/Engine/Source/Editor/BlueprintGraph/Classes/K2Node_EaseFunction.h
+++ b/Engine/Source/Editor/BlueprintGraph/Classes/K2Node_EaseFunction.h
@@ -10,7 +10,7 @@
 USTRUCT()
 struct FCustomPin
 {
-	GENERATED_USTRUCT_BODY()
+	GENERATED_BODY()
 
 	/** Name of the pin within the current node */
 	UPROPERTY()
@@ -28,7 +28,9 @@ struct FCustomPin
 UCLASS(MinimalAPI)
 class UK2Node_EaseFunction : public UK2Node
 {
-	GENERATED_UCLASS_BODY()
+	GENERATED_BODY()
+public:
+	BLUEPRINTGRAPH_API UK2Node_EaseFunction(const FObjectInitializer& ObjectInitializer = FObjectInitializer::Get());
 
 	// Begin UEdGraphNode interface.
 	virtual void AllocateDefaultPins() override;
diff --git a/Engine/Source/Editor/BlueprintGraph/Classes/K2Node_EditablePinBase.h b/Engine/Source/Editor/BlueprintGraph/Classes/K2Node_EditablePinBase.h
--- a/Engine/Source/Editor/BlueprintGraph/Classes/K2Node_EditablePinBase.h
+++ b/Engine/Source/Editor/BlueprintGraph/Classes/K2Node_EditablePinBase.h
@@ -7,7 +7,7 @@
 USTRUCT()
 struct FUserPinInfo
 {
-	GENERATED_USTRUCT_BODY()
+	GENERATED_BODY()
 
 	/** The name of the pin, as defined by the user */
 	UPROPERTY()
@@ -48,7 +48,7 @@ struct FUserPinInfo
 USTRUCT()
 struct FKismetUserDeclaredFunctionMetadata
 {
-	GENERATED_USTRUCT_BODY()
+	GENERATED_BODY()
 
 	UPROPERTY()
 	FString ToolTip;
@@ -77,7 +77,9 @@ public:
 UCLASS(abstract, MinimalAPI)
 class UK2Node_EditablePinBase : public UK2Node
 {
-	GENERATED_UCLASS_BODY()
+	GENERATED_BODY()
+public:
+	BLUEPRINTGRAPH_API UK2Node_EditablePinBase(const FObjectInitializer& ObjectInitializer = FObjectInitializer::Get());
 
 	/** Whether or not this entry node should be user-editable with the function editor */
 	UPROPERTY()
diff --git a/Engine/Source/Editor/BlueprintGraph/Classes/K2Node_EnumEquality.h b/Engine/Source/Editor/BlueprintGraph/Classes/K2Node_EnumEquality.h
--- a/Engine/Source/Editor/BlueprintGraph/Classes/K2Node_EnumEquality.h
+++ b/Engine/Source/Editor/BlueprintGraph/Classes/K2Node_EnumEquality.h
@@ -8,7 +8,9 @@
 UCLASS(MinimalAPI)
 class UK2Node_EnumEquality : public UK2Node
 {
-	GENERATED_UCLASS_BODY()
+	GENERATED_BODY()
+public:
+	BLUEPRINTGRAPH_API UK2Node_EnumEquality(const FObjectInitializer& ObjectInitializer = FObjectInitializer::Get());
 
 	// Begin UEdGraphNode interface
 	virtual void AllocateDefaultPins() override;
diff --git a/Engine/Source/Editor/BlueprintGraph/Classes/K2Node_EnumLiteral.h b/Engine/Source/Editor/BlueprintGraph/Classes/K2Node_EnumLiteral.h
--- a/Engine/Source/Editor/BlueprintGraph/Classes/K2Node_EnumLiteral.h
+++ b/Engine/Source/Editor/BlueprintGraph/Classes/K2Node_EnumLiteral.h
@@ -9,7 +9,9 @@
 UCLASS(MinimalAPI)
 class UK2Node_EnumLiteral : public UK2Node, public INodeDependingOnEnumInterface
 {
-	GENERATED_UCLASS_BODY()
+	GENERATED_BODY()
+public:
+	BLUEPRINTGRAPH_API UK2Node_EnumLiteral(const FObjectInitializer& ObjectInitializer = FObjectInitializer::Get());
 
 	UPROPERTY()
 	UEnum* Enum;
diff --git a/Engine/Source/Editor/BlueprintGraph/Classes/K2Node_Event.h b/Engine/Source/Editor/BlueprintGraph/Classes/K2Node_Event.h
--- a/Engine/Source/Editor/BlueprintGraph/Classes/K2Node_Event.h
+++ b/Engine/Source/Editor/BlueprintGraph/Classes/K2Node_Event.h
@@ -11,7 +11,9 @@
 UCLASS(MinimalAPI)
 class UK2Node_Event : public UK2Node_EditablePinBase
 {
-	GENERATED_UCLASS_BODY()
+	GENERATED_BODY()
+public:
+	BLUEPRINTGRAPH_API UK2Node_Event(const FObjectInitializer& ObjectInitializer = FObjectInitializer::Get());
 	BLUEPRINTGRAPH_API static const FString DelegateOutputName;
 
 	/** Name of function signature that this event implements */
diff --git a/Engine/Source/Editor/BlueprintGraph/Classes/K2Node_ExecutionSequence.h b/Engine/Source/Editor/BlueprintGraph/Classes/K2Node_ExecutionSequence.h
--- a/Engine/Source/Editor/BlueprintGraph/Classes/K2Node_ExecutionSequence.h
+++ b/Engine/Source/Editor/BlueprintGraph/Classes/K2Node_ExecutionSequence.h
@@ -8,7 +8,9 @@
 UCLASS(MinimalAPI)
 class UK2Node_ExecutionSequence : public UK2Node
 {
-	GENERATED_UCLASS_BODY()
+	GENERATED_BODY()
+public:
+	BLUEPRINTGRAPH_API UK2Node_ExecutionSequence(const FObjectInitializer& ObjectInitializer = FObjectInitializer::Get());
 
 	// Begin UEdGraphNode interface
 	virtual void AllocateDefaultPins() override;
diff --git a/Engine/Source/Editor/BlueprintGraph/Classes/K2Node_ForEachElementInEnum.h b/Engine/Source/Editor/BlueprintGraph/Classes/K2Node_ForEachElementInEnum.h
--- a/Engine/Source/Editor/BlueprintGraph/Classes/K2Node_ForEachElementInEnum.h
+++ b/Engine/Source/Editor/BlueprintGraph/Classes/K2Node_ForEachElementInEnum.h
@@ -8,7 +8,9 @@
 UCLASS(MinimalAPI)
 class UK2Node_ForEachElementInEnum : public UK2Node
 {
-	GENERATED_UCLASS_BODY()
+	GENERATED_BODY()
+public:
+	BLUEPRINTGRAPH_API UK2Node_ForEachElementInEnum(const FObjectInitializer& ObjectInitializer = FObjectInitializer::Get());
 
 	UPROPERTY()
 	UEnum* Enum;
diff --git a/Engine/Source/Editor/BlueprintGraph/Classes/K2Node_FormatText.h b/Engine/Source/Editor/BlueprintGraph/Classes/K2Node_FormatText.h
--- a/Engine/Source/Editor/BlueprintGraph/Classes/K2Node_FormatText.h
+++ b/Engine/Source/Editor/BlueprintGraph/Classes/K2Node_FormatText.h
@@ -1,4 +1,4 @@
-﻿// Copyright 1998-2015 Epic Games, Inc. All Rights Reserved.
+// Copyright 1998-2015 Epic Games, Inc. All Rights Reserved.
 
 #pragma once
 #include "K2Node_FormatText.generated.h"
@@ -6,7 +6,9 @@
 UCLASS(MinimalAPI)
 class UK2Node_FormatText : public UK2Node
 {
-	GENERATED_UCLASS_BODY()
+	GENERATED_BODY()
+public:
+	BLUEPRINTGRAPH_API UK2Node_FormatText(const FObjectInitializer& ObjectInitializer = FObjectInitializer::Get());
 
 	// UObject interface
 	virtual void PostLoad() override;
diff --git a/Engine/Source/Editor/BlueprintGraph/Classes/K2Node_FunctionEntry.h b/Engine/Source/Editor/BlueprintGraph/Classes/K2Node_FunctionEntry.h
--- a/Engine/Source/Editor/BlueprintGraph/Classes/K2Node_FunctionEntry.h
+++ b/Engine/Source/Editor/BlueprintGraph/Classes/K2Node_FunctionEntry.h
@@ -7,7 +7,9 @@
 UCLASS(MinimalAPI)
 class UK2Node_FunctionEntry : public UK2Node_FunctionTerminator
 {
-	GENERATED_UCLASS_BODY()
+	GENERATED_BODY()
+public:
+	BLUEPRINTGRAPH_API UK2Node_FunctionEntry(const FObjectInitializer& ObjectInitializer = FObjectInitializer::Get());
 
 	/** If specified, the function that is created for this entry point will have this name.  Otherwise, it will have the function signature's name */
 	UPROPERTY()
diff --git a/Engine/Source/Editor/BlueprintGraph/Classes/K2Node_FunctionResult.h b/Engine/Source/Editor/BlueprintGraph/Classes/K2Node_FunctionResult.h
--- a/Engine/Source/Editor/BlueprintGraph/Classes/K2Node_FunctionResult.h
+++ b/Engine/Source/Editor/BlueprintGraph/Classes/K2Node_FunctionResult.h
@@ -8,7 +8,9 @@
 UCLASS(MinimalAPI)
 class UK2Node_FunctionResult : public UK2Node_FunctionTerminator
 {
-	GENERATED_UCLASS_BODY()
+	GENERATED_BODY()
+public:
+	BLUEPRINTGRAPH_API UK2Node_FunctionResult(const FObjectInitializer& ObjectInitializer = FObjectInitializer::Get());
 
 	// Begin UEdGraphNode interface
 	virtual void AllocateDefaultPins() override;
diff --git a/Engine/Source/Editor/BlueprintGraph/Classes/K2Node_FunctionTerminator.h b/Engine/Source/Editor/BlueprintGraph/Classes/K2Node_FunctionTerminator.h
--- a/Engine/Source/Editor/BlueprintGraph/Classes/K2Node_FunctionTerminator.h
+++ b/Engine/Source/Editor/BlueprintGraph/Classes/K2Node_FunctionTerminator.h
@@ -8,7 +8,9 @@
 UCLASS(abstract, MinimalAPI)
 class UK2Node_FunctionTerminator : public UK2Node_EditablePinBase
 {
-	GENERATED_UCLASS_BODY()
+	GENERATED_BODY()
+public:
+	BLUEPRINTGRAPH_API UK2Node_FunctionTerminator(const FObjectInitializer& ObjectInitializer = FObjectInitializer::Get());
 
 	/** 
 	 * The source class that defines the signature, if it is getting that from elsewhere (e.g. interface, base class etc). 
diff --git a/Engine/Source/Editor/BlueprintGraph/Classes/K2Node_GetDataTableRow.h b/Engine/Source/Editor/BlueprintGraph/Classes/K2Node_GetDataTableRow.h
--- a/Engine/Source/Editor/BlueprintGraph/Classes/K2Node_GetDataTableRow.h
+++ b/Engine/Source/Editor/BlueprintGraph/Classes/K2Node_GetDataTableRow.h
@@ -10,7 +10,9 @@
 UCLASS()
 class BLUEPRINTGRAPH_API UK2Node_GetDataTableRow : public UK2Node
 {
-	GENERATED_UCLASS_BODY()
+	GENERATED_BODY()
+public:
+	UK2Node_GetDataTableRow(const FObjectInitializer& ObjectInitializer = FObjectInitializer::Get());
 
 	// Begin UEdGraphNode interface.
 	virtual void AllocateDefaultPins() override;
diff --git a/Engine/Source/Editor/BlueprintGraph/Classes/K2Node_GetEnumeratorName.h b/Engine/Source/Editor/BlueprintGraph/Classes/K2Node_GetEnumeratorName.h
--- a/Engine/Source/Editor/BlueprintGraph/Classes/K2Node_GetEnumeratorName.h
+++ b/Engine/Source/Editor/BlueprintGraph/Classes/K2Node_GetEnumeratorName.h
@@ -7,7 +7,9 @@
 UCLASS(MinimalAPI)
 class UK2Node_GetEnumeratorName : public UK2Node
 {
-	GENERATED_UCLASS_BODY()
+	GENERATED_BODY()
+public:
+	BLUEPRINTGRAPH_API UK2Node_GetEnumeratorName(const FObjectInitializer& ObjectInitializer = FObjectInitializer::Get());
 
 	static FString EnumeratorPinName;
 
diff --git a/Engine/Source/Editor/BlueprintGraph/Classes/K2Node_GetEnumeratorNameAsString.h b/Engine/Source/Editor/BlueprintGraph/Classes/K2Node_GetEnumeratorNameAsString.h
--- a/Engine/Source/Editor/BlueprintGraph/Classes/K2Node_GetEnumeratorNameAsString.h
+++ b/Engine/Source/Editor/BlueprintGraph/Classes/K2Node_GetEnumeratorNameAsString.h
@@ -7,7 +7,9 @@
 UCLASS(MinimalAPI)
 class UK2Node_GetEnumeratorNameAsString : public UK2Node_GetEnumeratorName
 {
-	GENERATED_UCLASS_BODY()
+	GENERATED_BODY()
+public:
+	BLUEPRINTGRAPH_API UK2Node_GetEnumeratorNameAsString(const FObjectInitializer& ObjectInitializer = FObjectInitializer::Get());
 
 	// Begin UEdGraphNode interface
 	virtual void AllocateDefaultPins() override;
diff --git a/Engine/Source/Editor/BlueprintGraph/Classes/K2Node_GetInputAxisKeyValue.h b/Engine/Source/Editor/BlueprintGraph/Classes/K2Node_GetInputAxisKeyValue.h
--- a/Engine/Source/Editor/BlueprintGraph/Classes/K2Node_GetInputAxisKeyValue.h
+++ b/Engine/Source/Editor/BlueprintGraph/Classes/K2Node_GetInputAxisKeyValue.h
@@ -8,7 +8,9 @@
 UCLASS(MinimalAPI)
 class UK2Node_GetInputAxisKeyValue : public UK2Node_CallFunction
 {
-	GENERATED_UCLASS_BODY()
+	GENERATED_BODY()
+public:
+	BLUEPRINTGRAPH_API UK2Node_GetInputAxisKeyValue(const FObjectInitializer& ObjectInitializer = FObjectInitializer::Get());
 
 	UPROPERTY()
 	FKey InputAxisKey;
@@ -46,4 +48,4 @@ private:
 	/** Constructing FText strings can be costly, so we cache the node's title/tooltip */
 	FNodeTextCache CachedTooltip;
 	FNodeTextCache CachedNodeTitle;
-};
\ No newline at end of file
+};
diff --git a/Engine/Source/Editor/BlueprintGraph/Classes/K2Node_GetInputAxisValue.h b/Engine/Source/Editor/BlueprintGraph/Classes/K2Node_GetInputAxisValue.h
--- a/Engine/Source/Editor/BlueprintGraph/Classes/K2Node_GetInputAxisValue.h
+++ b/Engine/Source/Editor/BlueprintGraph/Classes/K2Node_GetInputAxisValue.h
@@ -8,7 +8,9 @@
 UCLASS(MinimalAPI)
 class UK2Node_GetInputAxisValue : public UK2Node_CallFunction
 {
-	GENERATED_UCLASS_BODY()
+	GENERATED_BODY()
+public:
+	BLUEPRINTGRAPH_API UK2Node_GetInputAxisValue(const FObjectInitializer& ObjectInitializer = FObjectInitializer::Get());
 
 	UPROPERTY()
 	FName InputAxisName;
diff --git a/Engine/Source/Editor/BlueprintGraph/Classes/K2Node_GetInputVectorAxisValue.h b/Engine/Source/Editor/BlueprintGraph/Classes/K2Node_GetInputVectorAxisValue.h
--- a/Engine/Source/Editor/BlueprintGraph/Classes/K2Node_GetInputVectorAxisValue.h
+++ b/Engine/Source/Editor/BlueprintGraph/Classes/K2Node_GetInputVectorAxisValue.h
@@ -7,7 +7,9 @@
 UCLASS(MinimalAPI)
 class UK2Node_GetInputVectorAxisValue : public UK2Node_GetInputAxisKeyValue
 {
-	GENERATED_UCLASS_BODY()
+	GENERATED_BODY()
+public:
+	BLUEPRINTGRAPH_API UK2Node_GetInputVectorAxisValue(const FObjectInitializer& ObjectInitializer = FObjectInitializer::Get());
 
 	// Begin EdGraphNode interface
 	virtual FText GetTooltipText() const override;
diff --git a/Engine/Source/Editor/BlueprintGraph/Classes/K2Node_GetNumEnumEntries.h b/Engine/Source/Editor/BlueprintGraph/Classes/K2Node_GetNumEnumEntries.h
--- a/Engine/Source/Editor/BlueprintGraph/Classes/K2Node_GetNumEnumEntries.h
+++ b/Engine/Source/Editor/BlueprintGraph/Classes/K2Node_GetNumEnumEntries.h
@@ -8,7 +8,9 @@
 UCLASS(MinimalAPI)
 class UK2Node_GetNumEnumEntries : public UK2Node, public INodeDependingOnEnumInterface
 {
-	GENERATED_UCLASS_BODY()
+	GENERATED_BODY()
+public:
+	BLUEPRINTGRAPH_API UK2Node_GetNumEnumEntries(const FObjectInitializer& ObjectInitializer = FObjectInitializer::Get());
 
 	UPROPERTY()
 	UEnum* Enum;
diff --git a/Engine/Source/Editor/BlueprintGraph/Classes/K2Node_IfThenElse.h b/Engine/Source/Editor/BlueprintGraph/Classes/K2Node_IfThenElse.h
--- a/Engine/Source/Editor/BlueprintGraph/Classes/K2Node_IfThenElse.h
+++ b/Engine/Source/Editor/BlueprintGraph/Classes/K2Node_IfThenElse.h
@@ -8,7 +8,9 @@
 UCLASS(MinimalAPI)
 class UK2Node_IfThenElse : public UK2Node
 {
-	GENERATED_UCLASS_BODY()
+	GENERATED_BODY()
+public:
+	BLUEPRINTGRAPH_API UK2Node_IfThenElse(const FObjectInitializer& ObjectInitializer = FObjectInitializer::Get());
 
 	// Begin UEdGraphNode interface
 	virtual void AllocateDefaultPins() override;
diff --git a/Engine/Source/Editor/BlueprintGraph/Classes/K2Node_InputAction.h b/Engine/Source/Editor/BlueprintGraph/Classes/K2Node_InputAction.h
--- a/Engine/Source/Editor/BlueprintGraph/Classes/K2Node_InputAction.h
+++ b/Engine/Source/Editor/BlueprintGraph/Classes/K2Node_InputAction.h
@@ -9,7 +9,9 @@
 UCLASS(MinimalAPI)
 class UK2Node_InputAction : public UK2Node
 {
-	GENERATED_UCLASS_BODY()
+	GENERATED_BODY()
+public:
+	BLUEPRINTGRAPH_API UK2Node_InputAction(const FObjectInitializer& ObjectInitializer = FObjectInitializer::Get());
 
 	UPROPERTY()
 	FName InputActionName;
diff --git a/Engine/Source/Editor/BlueprintGraph/Classes/K2Node_InputActionEvent.h b/Engine/Source/Editor/BlueprintGraph/Classes/K2Node_InputActionEvent.h
--- a/Engine/Source/Editor/BlueprintGraph/Classes/K2Node_InputActionEvent.h
+++ b/Engine/Source/Editor/BlueprintGraph/Classes/K2Node_InputActionEvent.h
@@ -8,7 +8,9 @@
 UCLASS(MinimalAPI)
 class UK2Node_InputActionEvent : public UK2Node_Event
 {
-	GENERATED_UCLASS_BODY()
+	GENERATED_BODY()
+public:
+	BLUEPRINTGRAPH_API UK2Node_InputActionEvent(const FObjectInitializer& ObjectInitializer = FObjectInitializer::Get());
 
 	UPROPERTY()
 	FName InputActionName;
@@ -29,4 +31,4 @@ class UK2Node_InputActionEvent : public UK2Node_Event
 	virtual UClass* GetDynamicBindingClass() const override;
 	virtual void RegisterDynamicBinding(UDynamicBlueprintBinding* BindingObject) const override;
 	// End UK2Node interface
-};
\ No newline at end of file
+};
diff --git a/Engine/Source/Editor/BlueprintGraph/Classes/K2Node_InputAxisEvent.h b/Engine/Source/Editor/BlueprintGraph/Classes/K2Node_InputAxisEvent.h
--- a/Engine/Source/Editor/BlueprintGraph/Classes/K2Node_InputAxisEvent.h
+++ b/Engine/Source/Editor/BlueprintGraph/Classes/K2Node_InputAxisEvent.h
@@ -8,7 +8,9 @@
 UCLASS(MinimalAPI)
 class UK2Node_InputAxisEvent : public UK2Node_Event
 {
-	GENERATED_UCLASS_BODY()
+	GENERATED_BODY()
+public:
+	BLUEPRINTGRAPH_API UK2Node_InputAxisEvent(const FObjectInitializer& ObjectInitializer = FObjectInitializer::Get());
 
 	UPROPERTY()
 	FName InputAxisName;
diff --git a/Engine/Source/Editor/BlueprintGraph/Classes/K2Node_InputAxisKeyEvent.h b/Engine/Source/Editor/BlueprintGraph/Classes/K2Node_InputAxisKeyEvent.h
--- a/Engine/Source/Editor/BlueprintGraph/Classes/K2Node_InputAxisKeyEvent.h
+++ b/Engine/Source/Editor/BlueprintGraph/Classes/K2Node_InputAxisKeyEvent.h
@@ -7,7 +7,9 @@
 UCLASS(MinimalAPI)
 class UK2Node_InputAxisKeyEvent : public UK2Node_Event
 {
-	GENERATED_UCLASS_BODY()
+	GENERATED_BODY()
+public:
+	BLUEPRINTGRAPH_API UK2Node_InputAxisKeyEvent(const FObjectInitializer& ObjectInitializer = FObjectInitializer::Get());
 
 	UPROPERTY()
 	FKey AxisKey;
diff --git a/Engine/Source/Editor/BlueprintGraph/Classes/K2Node_InputKey.h b/Engine/Source/Editor/BlueprintGraph/Classes/K2Node_InputKey.h
--- a/Engine/Source/Editor/BlueprintGraph/Classes/K2Node_InputKey.h
+++ b/Engine/Source/Editor/BlueprintGraph/Classes/K2Node_InputKey.h
@@ -9,7 +9,9 @@
 UCLASS(MinimalAPI)
 class UK2Node_InputKey : public UK2Node
 {
-	GENERATED_UCLASS_BODY()
+	GENERATED_BODY()
+public:
+	BLUEPRINTGRAPH_API UK2Node_InputKey(const FObjectInitializer& ObjectInitializer = FObjectInitializer::Get());
 
 	// The key that is bound
 	UPROPERTY(EditAnywhere, Category="Input")
diff --git a/Engine/Source/Editor/BlueprintGraph/Classes/K2Node_InputKeyEvent.h b/Engine/Source/Editor/BlueprintGraph/Classes/K2Node_InputKeyEvent.h
--- a/Engine/Source/Editor/BlueprintGraph/Classes/K2Node_InputKeyEvent.h
+++ b/Engine/Source/Editor/BlueprintGraph/Classes/K2Node_InputKeyEvent.h
@@ -9,7 +9,9 @@
 UCLASS(MinimalAPI)
 class UK2Node_InputKeyEvent : public UK2Node_Event
 {
-	GENERATED_UCLASS_BODY()
+	GENERATED_BODY()
+public:
+	BLUEPRINTGRAPH_API UK2Node_InputKeyEvent(const FObjectInitializer& ObjectInitializer = FObjectInitializer::Get());
 
 	UPROPERTY()
 	FInputChord InputChord;
@@ -30,4 +32,4 @@ class UK2Node_InputKeyEvent : public UK2Node_Event
 	virtual UClass* GetDynamicBindingClass() const override;
 	virtual void RegisterDynamicBinding(UDynamicBlueprintBinding* BindingObject) const override;
 	// End UK2Node interface
-};
\ No newline at end of file
+};
diff --git a/Engine/Source/Editor/BlueprintGraph/Classes/K2Node_InputTouch.h b/Engine/Source/Editor/BlueprintGraph/Classes/K2Node_InputTouch.h
--- a/Engine/Source/Editor/BlueprintGraph/Classes/K2Node_InputTouch.h
+++ b/Engine/Source/Editor/BlueprintGraph/Classes/K2Node_InputTouch.h
@@ -8,7 +8,9 @@
 UCLASS(MinimalAPI)
 class UK2Node_InputTouch : public UK2Node
 {
-	GENERATED_UCLASS_BODY()
+	GENERATED_BODY()
+public:
+	BLUEPRINTGRAPH_API UK2Node_InputTouch(const FObjectInitializer& ObjectInitializer = FObjectInitializer::Get());
 
 	// Prevents actors with lower priority from handling this input
 	UPROPERTY(EditAnywhere, Category="Input")
diff --git a/Engine/Source/Editor/BlueprintGraph/Classes/K2Node_InputTouchEvent.h b/Engine/Source/Editor/BlueprintGraph/Classes/K2Node_InputTouchEvent.h
--- a/Engine/Source/Editor/BlueprintGraph/Classes/K2Node_InputTouchEvent.h
+++ b/Engine/Source/Editor/BlueprintGraph/Classes/K2Node_InputTouchEvent.h
@@ -8,7 +8,9 @@
 UCLASS(MinimalAPI)
 class UK2Node_InputTouchEvent : public UK2Node_Event
 {
-	GENERATED_UCLASS_BODY()
+	GENERATED_BODY()
+public:
+	BLUEPRINTGRAPH_API UK2Node_InputTouchEvent(const FObjectInitializer& ObjectInitializer = FObjectInitializer::Get());
 
 	UPROPERTY()
 	TEnumAsByte<EInputEvent> InputKeyEvent;
@@ -26,4 +28,4 @@ class UK2Node_InputTouchEvent : public UK2Node_Event
 	virtual UClass* GetDynamicBindingClass() const override;
 	virtual void RegisterDynamicBinding(UDynamicBlueprintBinding* BindingObject) const override;
 	// End UK2Node interface
-};
\ No newline at end of file
+};
diff --git a/Engine/Source/Editor/BlueprintGraph/Classes/K2Node_InputVectorAxisEvent.h b/Engine/Source/Editor/BlueprintGraph/Classes/K2Node_InputVectorAxisEvent.h
--- a/Engine/Source/Editor/BlueprintGraph/Classes/K2Node_InputVectorAxisEvent.h
+++ b/Engine/Source/Editor/BlueprintGraph/Classes/K2Node_InputVectorAxisEvent.h
@@ -7,7 +7,9 @@
 UCLASS(MinimalAPI)
 class UK2Node_InputVectorAxisEvent : public UK2Node_InputAxisKeyEvent
 {
-	GENERATED_UCLASS_BODY()
+	GENERATED_BODY()
+public:
+	BLUEPRINTGRAPH_API UK2Node_InputVectorAxisEvent(const FObjectInitializer& ObjectInitializer = FObjectInitializer::Get());
 
 	// UObject interface
 	virtual void Serialize(FArchive& Ar) override;
@@ -19,4 +21,4 @@ class UK2Node_InputVectorAxisEvent : public UK2Node_InputAxisKeyEvent
 	virtual void RegisterDynamicBinding(UDynamicBlueprintBinding* BindingObject) const override;
 	virtual void GetMenuActions(FBlueprintActionDatabaseRegistrar& ActionRegistrar) const override;
 	// End UK2Node interface
-};
\ No newline at end of file
+};
diff --git a/Engine/Source/Editor/BlueprintGraph/Classes/K2Node_Knot.h b/Engine/Source/Editor/BlueprintGraph/Classes/K2Node_Knot.h
--- a/Engine/Source/Editor/BlueprintGraph/Classes/K2Node_Knot.h
+++ b/Engine/Source/Editor/BlueprintGraph/Classes/K2Node_Knot.h
@@ -7,7 +7,9 @@
 UCLASS(MinimalAPI)
 class UK2Node_Knot : public UK2Node
 {
-	GENERATED_UCLASS_BODY()
+	GENERATED_BODY()
+public:
+	BLUEPRINTGRAPH_API UK2Node_Knot(const FObjectInitializer& ObjectInitializer = FObjectInitializer::Get());
 
 public:
 	// UEdGraphNode interface
diff --git a/Engine/Source/Editor/BlueprintGraph/Classes/K2Node_Literal.h b/Engine/Source/Editor/BlueprintGraph/Classes/K2Node_Literal.h
--- a/Engine/Source/Editor/BlueprintGraph/Classes/K2Node_Literal.h
+++ b/Engine/Source/Editor/BlueprintGraph/Classes/K2Node_Literal.h
@@ -8,7 +8,9 @@
 UCLASS(MinimalAPI)
 class UK2Node_Literal : public UK2Node
 {
-	GENERATED_UCLASS_BODY()
+	GENERATED_BODY()
+public:
+	BLUEPRINTGRAPH_API UK2Node_Literal(const FObjectInitializer& ObjectInitializer = FObjectInitializer::Get());
 
 private:
 	/** If this is an object reference literal, keep a reference here so that it can be updated as objects move around */
diff --git a/Engine/Source/Editor/BlueprintGraph/Classes/K2Node_LocalVariable.h b/Engine/Source/Editor/BlueprintGraph/Classes/K2Node_LocalVariable.h
--- a/Engine/Source/Editor/BlueprintGraph/Classes/K2Node_LocalVariable.h
+++ b/Engine/Source/Editor/BlueprintGraph/Classes/K2Node_LocalVariable.h
@@ -7,7 +7,9 @@
 UCLASS(MinimalAPI, deprecated)
 class UDEPRECATED_K2Node_LocalVariable : public UK2Node_TemporaryVariable
 {
-	GENERATED_UCLASS_BODY()
+	GENERATED_BODY()
+public:
+	BLUEPRINTGRAPH_API UDEPRECATED_K2Node_LocalVariable(const FObjectInitializer& ObjectInitializer = FObjectInitializer::Get());
 
 	/** If this is not an override, allow user to specify a name for the function created by this entry point */
 	UPROPERTY()
@@ -40,4 +42,4 @@ class UDEPRECATED_K2Node_LocalVariable : public UK2Node_TemporaryVariable
 	 * @param InVariableType		The type to assign this local variable to
 	 */
 	BLUEPRINTGRAPH_API void ChangeVariableType(const FEdGraphPinType& InVariableType);
-};
\ No newline at end of file
+};
diff --git a/Engine/Source/Editor/BlueprintGraph/Classes/K2Node_MacroInstance.h b/Engine/Source/Editor/BlueprintGraph/Classes/K2Node_MacroInstance.h
--- a/Engine/Source/Editor/BlueprintGraph/Classes/K2Node_MacroInstance.h
+++ b/Engine/Source/Editor/BlueprintGraph/Classes/K2Node_MacroInstance.h
@@ -10,7 +10,9 @@
 UCLASS(MinimalAPI)
 class UK2Node_MacroInstance : public UK2Node_Tunnel
 {
-	GENERATED_UCLASS_BODY()
+	GENERATED_BODY()
+public:
+	BLUEPRINTGRAPH_API UK2Node_MacroInstance(const FObjectInitializer& ObjectInitializer = FObjectInitializer::Get());
 
 private:
 	/** A macro is like a composite node, except that the associated graph lives
diff --git a/Engine/Source/Editor/BlueprintGraph/Classes/K2Node_MakeArray.h b/Engine/Source/Editor/BlueprintGraph/Classes/K2Node_MakeArray.h
--- a/Engine/Source/Editor/BlueprintGraph/Classes/K2Node_MakeArray.h
+++ b/Engine/Source/Editor/BlueprintGraph/Classes/K2Node_MakeArray.h
@@ -8,7 +8,9 @@
 UCLASS(MinimalAPI)
 class UK2Node_MakeArray : public UK2Node
 {
-	GENERATED_UCLASS_BODY()
+	GENERATED_BODY()
+public:
+	BLUEPRINTGRAPH_API UK2Node_MakeArray(const FObjectInitializer& ObjectInitializer = FObjectInitializer::Get());
 
 	/** The number of input pins to generate for this node */
 	UPROPERTY()
diff --git a/Engine/Source/Editor/BlueprintGraph/Classes/K2Node_MakeStruct.h b/Engine/Source/Editor/BlueprintGraph/Classes/K2Node_MakeStruct.h
--- a/Engine/Source/Editor/BlueprintGraph/Classes/K2Node_MakeStruct.h
+++ b/Engine/Source/Editor/BlueprintGraph/Classes/K2Node_MakeStruct.h
@@ -9,7 +9,9 @@
 UCLASS(MinimalAPI)
 class UK2Node_MakeStruct : public UK2Node_StructMemberSet
 {
-	GENERATED_UCLASS_BODY()
+	GENERATED_BODY()
+public:
+	BLUEPRINTGRAPH_API UK2Node_MakeStruct(const FObjectInitializer& ObjectInitializer = FObjectInitializer::Get());
 
 	BLUEPRINTGRAPH_API static bool CanBeMade(const UScriptStruct* Struct);
 	static bool CanBeExposed(const UProperty* Property);
diff --git a/Engine/Source/Editor/BlueprintGraph/Classes/K2Node_MathExpression.h b/Engine/Source/Editor/BlueprintGraph/Classes/K2Node_MathExpression.h
--- a/Engine/Source/Editor/BlueprintGraph/Classes/K2Node_MathExpression.h
+++ b/Engine/Source/Editor/BlueprintGraph/Classes/K2Node_MathExpression.h
@@ -16,7 +16,9 @@
 UCLASS()
 class BLUEPRINTGRAPH_API UK2Node_MathExpression : public UK2Node_Composite
 {
-	GENERATED_UCLASS_BODY()
+	GENERATED_BODY()
+public:
+	UK2Node_MathExpression(const FObjectInitializer& ObjectInitializer = FObjectInitializer::Get());
 
 public:
 
diff --git a/Engine/Source/Editor/BlueprintGraph/Classes/K2Node_MatineeController.h b/Engine/Source/Editor/BlueprintGraph/Classes/K2Node_MatineeController.h
--- a/Engine/Source/Editor/BlueprintGraph/Classes/K2Node_MatineeController.h
+++ b/Engine/Source/Editor/BlueprintGraph/Classes/K2Node_MatineeController.h
@@ -8,7 +8,9 @@
 UCLASS(MinimalAPI)
 class UK2Node_MatineeController : public UK2Node
 {
-	GENERATED_UCLASS_BODY()
+	GENERATED_BODY()
+public:
+	BLUEPRINTGRAPH_API UK2Node_MatineeController(const FObjectInitializer& ObjectInitializer = FObjectInitializer::Get());
 
 	/** The matinee actor in the level that this node controls */
 	UPROPERTY(EditAnywhere, Category=K2Node_MatineeController)
diff --git a/Engine/Source/Editor/BlueprintGraph/Classes/K2Node_Message.h b/Engine/Source/Editor/BlueprintGraph/Classes/K2Node_Message.h
--- a/Engine/Source/Editor/BlueprintGraph/Classes/K2Node_Message.h
+++ b/Engine/Source/Editor/BlueprintGraph/Classes/K2Node_Message.h
@@ -9,7 +9,9 @@
 UCLASS(MinimalAPI)
 class UK2Node_Message : public UK2Node_CallFunction
 {
-	GENERATED_UCLASS_BODY()
+	GENERATED_BODY()
+public:
+	BLUEPRINTGRAPH_API UK2Node_Message(const FObjectInitializer& ObjectInitializer = FObjectInitializer::Get());
 
 
 	// Begin UEdGraphNode interface.
diff --git a/Engine/Source/Editor/BlueprintGraph/Classes/K2Node_MultiGate.h b/Engine/Source/Editor/BlueprintGraph/Classes/K2Node_MultiGate.h
--- a/Engine/Source/Editor/BlueprintGraph/Classes/K2Node_MultiGate.h
+++ b/Engine/Source/Editor/BlueprintGraph/Classes/K2Node_MultiGate.h
@@ -7,7 +7,9 @@
 UCLASS(MinimalAPI)
 class UK2Node_MultiGate : public UK2Node_ExecutionSequence
 {
-	GENERATED_UCLASS_BODY()
+	GENERATED_BODY()
+public:
+	BLUEPRINTGRAPH_API UK2Node_MultiGate(const FObjectInitializer& ObjectInitializer = FObjectInitializer::Get());
 
 	/** Reference to the integer that contains */
 	UPROPERTY(transient)
diff --git a/Engine/Source/Editor/BlueprintGraph/Classes/K2Node_PureAssignmentStatement.h b/Engine/Source/Editor/BlueprintGraph/Classes/K2Node_PureAssignmentStatement.h
--- a/Engine/Source/Editor/BlueprintGraph/Classes/K2Node_PureAssignmentStatement.h
+++ b/Engine/Source/Editor/BlueprintGraph/Classes/K2Node_PureAssignmentStatement.h
@@ -7,7 +7,9 @@
 UCLASS(MinimalAPI)
 class UK2Node_PureAssignmentStatement : public UK2Node
 {
-	GENERATED_UCLASS_BODY()
+	GENERATED_BODY()
+public:
+	BLUEPRINTGRAPH_API UK2Node_PureAssignmentStatement(const FObjectInitializer& ObjectInitializer = FObjectInitializer::Get());
 
 	// Name of the Variable pin for this node
 	static FString VariablePinName;
diff --git a/Engine/Source/Editor/BlueprintGraph/Classes/K2Node_RemoveDelegate.h b/Engine/Source/Editor/BlueprintGraph/Classes/K2Node_RemoveDelegate.h
--- a/Engine/Source/Editor/BlueprintGraph/Classes/K2Node_RemoveDelegate.h
+++ b/Engine/Source/Editor/BlueprintGraph/Classes/K2Node_RemoveDelegate.h
@@ -6,7 +6,9 @@
 UCLASS(MinimalAPI)
 class UK2Node_RemoveDelegate : public UK2Node_BaseMCDelegate
 {
-	GENERATED_UCLASS_BODY()
+	GENERATED_BODY()
+public:
+	BLUEPRINTGRAPH_API UK2Node_RemoveDelegate(const FObjectInitializer& ObjectInitializer = FObjectInitializer::Get());
 
 public:
 	// UEdGraphNode interface
diff --git a/Engine/Source/Editor/BlueprintGraph/Classes/K2Node_Select.h b/Engine/Source/Editor/BlueprintGraph/Classes/K2Node_Select.h
--- a/Engine/Source/Editor/BlueprintGraph/Classes/K2Node_Select.h
+++ b/Engine/Source/Editor/BlueprintGraph/Classes/K2Node_Select.h
@@ -9,7 +9,9 @@
 UCLASS(MinimalAPI)
 class UK2Node_Select : public UK2Node, public INodeDependingOnEnumInterface
 {
-	GENERATED_UCLASS_BODY()
+	GENERATED_BODY()
+public:
+	BLUEPRINTGRAPH_API UK2Node_Select(const FObjectInitializer& ObjectInitializer = FObjectInitializer::Get());
 
 	/** The number of selectable options this node currently has */
 	UPROPERTY()
diff --git a/Engine/Source/Editor/BlueprintGraph/Classes/K2Node_Self.h b/Engine/Source/Editor/BlueprintGraph/Classes/K2Node_Self.h
--- a/Engine/Source/Editor/BlueprintGraph/Classes/K2Node_Self.h
+++ b/Engine/Source/Editor/BlueprintGraph/Classes/K2Node_Self.h
@@ -8,7 +8,9 @@
 UCLASS(MinimalAPI)
 class UK2Node_Self : public UK2Node
 {
-	GENERATED_UCLASS_BODY()
+	GENERATED_BODY()
+public:
+	BLUEPRINTGRAPH_API UK2Node_Self(const FObjectInitializer& ObjectInitializer = FObjectInitializer::Get());
 
 	// Begin UEdGraphNode interface
 	virtual FText GetNodeTitle(ENodeTitleType::Type TitleType) const override;
diff --git a/Engine/Source/Editor/BlueprintGraph/Classes/K2Node_SetFieldsInStruct.h b/Engine/Source/Editor/BlueprintGraph/Classes/K2Node_SetFieldsInStruct.h
--- a/Engine/Source/Editor/BlueprintGraph/Classes/K2Node_SetFieldsInStruct.h
+++ b/Engine/Source/Editor/BlueprintGraph/Classes/K2Node_SetFieldsInStruct.h
@@ -9,7 +9,9 @@
 UCLASS(MinimalAPI)
 class UK2Node_SetFieldsInStruct : public UK2Node_MakeStruct
 {
-	GENERATED_UCLASS_BODY()
+	GENERATED_BODY()
+public:
+	BLUEPRINTGRAPH_API UK2Node_SetFieldsInStruct(const FObjectInitializer& ObjectInitializer = FObjectInitializer::Get());
 	// Begin UEdGraphNode interface
 	virtual void AllocateDefaultPins() override;
 	virtual FText GetNodeTitle(ENodeTitleType::Type TitleType) const override;
diff --git a/Engine/Source/Editor/BlueprintGraph/Classes/K2Node_SpawnActor.h b/Engine/Source/Editor/BlueprintGraph/Classes/K2Node_SpawnActor.h
--- a/Engine/Source/Editor/BlueprintGraph/Classes/K2Node_SpawnActor.h
+++ b/Engine/Source/Editor/BlueprintGraph/Classes/K2Node_SpawnActor.h
@@ -9,7 +9,9 @@
 UCLASS(MinimalAPI)
 class UK2Node_SpawnActor : public UK2Node
 {
-	GENERATED_UCLASS_BODY()
+	GENERATED_BODY()
+public:
+	BLUEPRINTGRAPH_API UK2Node_SpawnActor(const FObjectInitializer& ObjectInitializer = FObjectInitializer::Get());
 
 	// Begin UEdGraphNode interface.
 	virtual void AllocateDefaultPins() override;
diff --git a/Engine/Source/Editor/BlueprintGraph/Classes/K2Node_SpawnActorFromClass.h b/Engine/Source/Editor/BlueprintGraph/Classes/K2Node_SpawnActorFromClass.h
--- a/Engine/Source/Editor/BlueprintGraph/Classes/K2Node_SpawnActorFromClass.h
+++ b/Engine/Source/Editor/BlueprintGraph/Classes/K2Node_SpawnActorFromClass.h
@@ -9,7 +9,9 @@
 UCLASS()
 class BLUEPRINTGRAPH_API UK2Node_SpawnActorFromClass : public UK2Node
 {
-	GENERATED_UCLASS_BODY()
+	GENERATED_BODY()
+public:
+	UK2Node_SpawnActorFromClass(const FObjectInitializer& ObjectInitializer = FObjectInitializer::Get());
 
 	// Begin UEdGraphNode interface.
 	virtual void AllocateDefaultPins() override;
diff --git a/Engine/Source/Editor/BlueprintGraph/Classes/K2Node_StructMemberGet.h b/Engine/Source/Editor/BlueprintGraph/Classes/K2Node_StructMemberGet.h
--- a/Engine/Source/Editor/BlueprintGraph/Classes/K2Node_StructMemberGet.h
+++ b/Engine/Source/Editor/BlueprintGraph/Classes/K2Node_StructMemberGet.h
@@ -9,7 +9,9 @@
 UCLASS(MinimalAPI)
 class UK2Node_StructMemberGet : public UK2Node_StructOperation
 {
-	GENERATED_UCLASS_BODY()
+	GENERATED_BODY()
+public:
+	BLUEPRINTGRAPH_API UK2Node_StructMemberGet(const FObjectInitializer& ObjectInitializer = FObjectInitializer::Get());
 
 	UPROPERTY(EditAnywhere, Category=PinOptions, EditFixedSize)
 	TArray<FOptionalPinFromProperty> ShowPinForProperties;
diff --git a/Engine/Source/Editor/BlueprintGraph/Classes/K2Node_StructMemberSet.h b/Engine/Source/Editor/BlueprintGraph/Classes/K2Node_StructMemberSet.h
--- a/Engine/Source/Editor/BlueprintGraph/Classes/K2Node_StructMemberSet.h
+++ b/Engine/Source/Editor/BlueprintGraph/Classes/K2Node_StructMemberSet.h
@@ -9,7 +9,9 @@
 UCLASS(MinimalAPI)
 class UK2Node_StructMemberSet : public UK2Node_StructOperation
 {
-	GENERATED_UCLASS_BODY()
+	GENERATED_BODY()
+public:
+	BLUEPRINTGRAPH_API UK2Node_StructMemberSet(const FObjectInitializer& ObjectInitializer = FObjectInitializer::Get());
 
 	UPROPERTY(EditAnywhere, Category=PinOptions, EditFixedSize)
 	TArray<FOptionalPinFromProperty> ShowPinForProperties;
diff --git a/Engine/Source/Editor/BlueprintGraph/Classes/K2Node_StructOperation.h b/Engine/Source/Editor/BlueprintGraph/Classes/K2Node_StructOperation.h
--- a/Engine/Source/Editor/BlueprintGraph/Classes/K2Node_StructOperation.h
+++ b/Engine/Source/Editor/BlueprintGraph/Classes/K2Node_StructOperation.h
@@ -8,7 +8,9 @@
 UCLASS(MinimalAPI, abstract)
 class UK2Node_StructOperation : public UK2Node_Variable
 {
-	GENERATED_UCLASS_BODY()
+	GENERATED_BODY()
+public:
+	BLUEPRINTGRAPH_API UK2Node_StructOperation(const FObjectInitializer& ObjectInitializer = FObjectInitializer::Get());
 
 	/** Class that this variable is defined in.  */
 	UPROPERTY()
diff --git a/Engine/Source/Editor/BlueprintGraph/Classes/K2Node_Switch.h b/Engine/Source/Editor/BlueprintGraph/Classes/K2Node_Switch.h
--- a/Engine/Source/Editor/BlueprintGraph/Classes/K2Node_Switch.h
+++ b/Engine/Source/Editor/BlueprintGraph/Classes/K2Node_Switch.h
@@ -8,7 +8,9 @@
 UCLASS(MinimalAPI, abstract)
 class UK2Node_Switch : public UK2Node
 {
-	GENERATED_UCLASS_BODY()
+	GENERATED_BODY()
+public:
+	BLUEPRINTGRAPH_API UK2Node_Switch(const FObjectInitializer& ObjectInitializer = FObjectInitializer::Get());
 
 	/** If true switch has a default pin */
 	UPROPERTY(EditAnywhere, Category=PinOptions)
diff --git a/Engine/Source/Editor/BlueprintGraph/Classes/K2Node_SwitchEnum.h b/Engine/Source/Editor/BlueprintGraph/Classes/K2Node_SwitchEnum.h
--- a/Engine/Source/Editor/BlueprintGraph/Classes/K2Node_SwitchEnum.h
+++ b/Engine/Source/Editor/BlueprintGraph/Classes/K2Node_SwitchEnum.h
@@ -10,7 +10,9 @@
 UCLASS(MinimalAPI)
 class UK2Node_SwitchEnum : public UK2Node_Switch, public INodeDependingOnEnumInterface
 {
-	GENERATED_UCLASS_BODY()
+	GENERATED_BODY()
+public:
+	BLUEPRINTGRAPH_API UK2Node_SwitchEnum(const FObjectInitializer& ObjectInitializer = FObjectInitializer::Get());
 
 	/** Name of the enum being switched on */
 	UPROPERTY()
diff --git a/Engine/Source/Editor/BlueprintGraph/Classes/K2Node_SwitchInteger.h b/Engine/Source/Editor/BlueprintGraph/Classes/K2Node_SwitchInteger.h
--- a/Engine/Source/Editor/BlueprintGraph/Classes/K2Node_SwitchInteger.h
+++ b/Engine/Source/Editor/BlueprintGraph/Classes/K2Node_SwitchInteger.h
@@ -8,7 +8,9 @@
 UCLASS(MinimalAPI)
 class UK2Node_SwitchInteger : public UK2Node_Switch
 {
-	GENERATED_UCLASS_BODY()
+	GENERATED_BODY()
+public:
+	BLUEPRINTGRAPH_API UK2Node_SwitchInteger(const FObjectInitializer& ObjectInitializer = FObjectInitializer::Get());
 
 	/** Set the starting index for the node */
 	UPROPERTY(EditAnywhere, Category=PinOptions, meta=(NoSpinbox=true))
diff --git a/Engine/Source/Editor/BlueprintGraph/Classes/K2Node_SwitchName.h b/Engine/Source/Editor/BlueprintGraph/Classes/K2Node_SwitchName.h
--- a/Engine/Source/Editor/BlueprintGraph/Classes/K2Node_SwitchName.h
+++ b/Engine/Source/Editor/BlueprintGraph/Classes/K2Node_SwitchName.h
@@ -8,7 +8,9 @@
 UCLASS(MinimalAPI)
 class UK2Node_SwitchName : public UK2Node_Switch
 {
-	GENERATED_UCLASS_BODY()
+	GENERATED_BODY()
+public:
+	BLUEPRINTGRAPH_API UK2Node_SwitchName(const FObjectInitializer& ObjectInitializer = FObjectInitializer::Get());
 
 	UPROPERTY(EditAnywhere, Category=PinOptions)
 	TArray<FName> PinNames;
diff --git a/Engine/Source/Editor/BlueprintGraph/Classes/K2Node_SwitchString.h b/Engine/Source/Editor/BlueprintGraph/Classes/K2Node_SwitchString.h
--- a/Engine/Source/Editor/BlueprintGraph/Classes/K2Node_SwitchString.h
+++ b/Engine/Source/Editor/BlueprintGraph/Classes/K2Node_SwitchString.h
@@ -8,7 +8,9 @@
 UCLASS(MinimalAPI)
 class UK2Node_SwitchString : public UK2Node_Switch
 {
-	GENERATED_UCLASS_BODY()
+	GENERATED_BODY()
+public:
+	BLUEPRINTGRAPH_API UK2Node_SwitchString(const FObjectInitializer& ObjectInitializer = FObjectInitializer::Get());
 
 	UPROPERTY(EditAnywhere, Category=PinOptions)
 	TArray<FString> PinNames;
diff --git a/Engine/Source/Editor/BlueprintGraph/Classes/K2Node_TemporaryVariable.h b/Engine/Source/Editor/BlueprintGraph/Classes/K2Node_TemporaryVariable.h
--- a/Engine/Source/Editor/BlueprintGraph/Classes/K2Node_TemporaryVariable.h
+++ b/Engine/Source/Editor/BlueprintGraph/Classes/K2Node_TemporaryVariable.h
@@ -8,7 +8,9 @@
 UCLASS(MinimalAPI)
 class UK2Node_TemporaryVariable : public UK2Node
 {
-	GENERATED_UCLASS_BODY()
+	GENERATED_BODY()
+public:
+	BLUEPRINTGRAPH_API UK2Node_TemporaryVariable(const FObjectInitializer& ObjectInitializer = FObjectInitializer::Get());
 
 	UPROPERTY()
 	struct FEdGraphPinType VariableType;
diff --git a/Engine/Source/Editor/BlueprintGraph/Classes/K2Node_Timeline.h b/Engine/Source/Editor/BlueprintGraph/Classes/K2Node_Timeline.h
--- a/Engine/Source/Editor/BlueprintGraph/Classes/K2Node_Timeline.h
+++ b/Engine/Source/Editor/BlueprintGraph/Classes/K2Node_Timeline.h
@@ -8,7 +8,9 @@
 UCLASS(MinimalAPI)
 class UK2Node_Timeline : public UK2Node
 {
-	GENERATED_UCLASS_BODY()
+	GENERATED_BODY()
+public:
+	BLUEPRINTGRAPH_API UK2Node_Timeline(const FObjectInitializer& ObjectInitializer = FObjectInitializer::Get());
 
 	/** The name of the timeline. Used to name ONLY the member variable (Component). To obtain the name of timeline template use UTimelineTemplate::TimelineVariableNameToTemplateName */
 	UPROPERTY()
diff --git a/Engine/Source/Editor/BlueprintGraph/Classes/K2Node_Tunnel.h b/Engine/Source/Editor/BlueprintGraph/Classes/K2Node_Tunnel.h
--- a/Engine/Source/Editor/BlueprintGraph/Classes/K2Node_Tunnel.h
+++ b/Engine/Source/Editor/BlueprintGraph/Classes/K2Node_Tunnel.h
@@ -7,7 +7,9 @@
 UCLASS(MinimalAPI)
 class UK2Node_Tunnel : public UK2Node_EditablePinBase
 {
-	GENERATED_UCLASS_BODY()
+	GENERATED_BODY()
+public:
+	BLUEPRINTGRAPH_API UK2Node_Tunnel(const FObjectInitializer& ObjectInitializer = FObjectInitializer::Get());
 
 	// A tunnel node either has output pins that came from another tunnel's input pins, or vis versa
 	// Note: OutputSourceNode might be equal to InputSinkNode
diff --git a/Engine/Source/Editor/BlueprintGraph/Classes/K2Node_Variable.h b/Engine/Source/Editor/BlueprintGraph/Classes/K2Node_Variable.h
--- a/Engine/Source/Editor/BlueprintGraph/Classes/K2Node_Variable.h
+++ b/Engine/Source/Editor/BlueprintGraph/Classes/K2Node_Variable.h
@@ -22,7 +22,9 @@ namespace ESelfContextInfo
 UCLASS(abstract)
 class BLUEPRINTGRAPH_API UK2Node_Variable : public UK2Node
 {
-	GENERATED_UCLASS_BODY()
+	GENERATED_BODY()
+public:
+	UK2Node_Variable(const FObjectInitializer& ObjectInitializer = FObjectInitializer::Get());
 
 	/** Reference to variable we want to set/get */
 	UPROPERTY()
diff --git a/Engine/Source/Editor/BlueprintGraph/Classes/K2Node_VariableGet.h b/Engine/Source/Editor/BlueprintGraph/Classes/K2Node_VariableGet.h
--- a/Engine/Source/Editor/BlueprintGraph/Classes/K2Node_VariableGet.h
+++ b/Engine/Source/Editor/BlueprintGraph/Classes/K2Node_VariableGet.h
@@ -9,7 +9,9 @@
 UCLASS()
 class BLUEPRINTGRAPH_API UK2Node_VariableGet : public UK2Node_Variable
 {
-	GENERATED_UCLASS_BODY()
+	GENERATED_BODY()
+public:
+	UK2Node_VariableGet(const FObjectInitializer& ObjectInitializer = FObjectInitializer::Get());
 
 	// Begin UEdGraphNode interface
 	virtual void AllocateDefaultPins() override;
diff --git a/Engine/Source/Editor/BlueprintGraph/Classes/K2Node_VariableSet.h b/Engine/Source/Editor/BlueprintGraph/Classes/K2Node_VariableSet.h
--- a/Engine/Source/Editor/BlueprintGraph/Classes/K2Node_VariableSet.h
+++ b/Engine/Source/Editor/BlueprintGraph/Classes/K2Node_VariableSet.h
@@ -9,7 +9,9 @@
 UCLASS(MinimalAPI)
 class UK2Node_VariableSet : public UK2Node_Variable
 {
-	GENERATED_UCLASS_BODY()
+	GENERATED_BODY()
+public:
+	BLUEPRINTGRAPH_API UK2Node_VariableSet(const FObjectInitializer& ObjectInitializer = FObjectInitializer::Get());
 
 	// Begin UEdGraphNode interface
 	virtual void AllocateDefaultPins() override;
diff --git a/Engine/Source/Editor/BlueprintGraph/Classes/K2Node_VariableSetRef.h b/Engine/Source/Editor/BlueprintGraph/Classes/K2Node_VariableSetRef.h
--- a/Engine/Source/Editor/BlueprintGraph/Classes/K2Node_VariableSetRef.h
+++ b/Engine/Source/Editor/BlueprintGraph/Classes/K2Node_VariableSetRef.h
@@ -9,7 +9,9 @@
 UCLASS(MinimalAPI)
 class UK2Node_VariableSetRef : public UK2Node
 {
-	GENERATED_UCLASS_BODY()
+	GENERATED_BODY()
+public:
+	BLUEPRINTGRAPH_API UK2Node_VariableSetRef(const FObjectInitializer& ObjectInitializer = FObjectInitializer::Get());
 
 	// Begin UEdGraphNode interface
 	virtual void AllocateDefaultPins() override;
diff --git a/Engine/Source/Editor/BlueprintGraph/Classes/NodeDependingOnEnumInterface.h b/Engine/Source/Editor/BlueprintGraph/Classes/NodeDependingOnEnumInterface.h
--- a/Engine/Source/Editor/BlueprintGraph/Classes/NodeDependingOnEnumInterface.h
+++ b/Engine/Source/Editor/BlueprintGraph/Classes/NodeDependingOnEnumInterface.h
@@ -6,12 +6,15 @@
 UINTERFACE(meta=(CannotImplementInterfaceInBlueprint))
 class BLUEPRINTGRAPH_API UNodeDependingOnEnumInterface : public UInterface
 {
-	GENERATED_UINTERFACE_BODY()
+	GENERATED_BODY()
+public:
+	UNodeDependingOnEnumInterface(const FObjectInitializer& ObjectInitializer = FObjectInitializer::Get());
 };
 
 class BLUEPRINTGRAPH_API INodeDependingOnEnumInterface
 {
-	GENERATED_IINTERFACE_BODY()
+	GENERATED_BODY()
+public:
 
 	virtual class UEnum* GetEnum() const PURE_VIRTUAL(INodeDependingOnEnumInterface::GetEnum,return NULL;);
 
diff --git a/Engine/Source/Editor/BlueprintGraph/Public/BlueprintEditorSettings.h b/Engine/Source/Editor/BlueprintGraph/Public/BlueprintEditorSettings.h
--- a/Engine/Source/Editor/BlueprintGraph/Public/BlueprintEditorSettings.h
+++ b/Engine/Source/Editor/BlueprintGraph/Public/BlueprintEditorSettings.h
@@ -17,7 +17,9 @@ UCLASS(config=EditorUserSettings)
 class BLUEPRINTGRAPH_API UBlueprintEditorSettings
 	:	public UObject
 {
-	GENERATED_UCLASS_BODY()
+	GENERATED_BODY()
+public:
+	UBlueprintEditorSettings(const FObjectInitializer& ObjectInitializer = FObjectInitializer::Get());
 
 // Style Settings
 public:
diff --git a/Engine/Source/Editor/BlueprintGraph/Public/BlueprintNodeSpawner.h b/Engine/Source/Editor/BlueprintGraph/Public/BlueprintNodeSpawner.h
--- a/Engine/Source/Editor/BlueprintGraph/Public/BlueprintNodeSpawner.h
+++ b/Engine/Source/Editor/BlueprintGraph/Public/BlueprintNodeSpawner.h
@@ -42,7 +42,9 @@ struct FBlueprintActionUiSpec
 UCLASS(Transient)
 class BLUEPRINTGRAPH_API UBlueprintNodeSpawner : public UObject, public IBlueprintNodeBinder
 {
-	GENERATED_UCLASS_BODY()
+	GENERATED_BODY()
+public:
+	UBlueprintNodeSpawner(const FObjectInitializer& ObjectInitializer = FObjectInitializer::Get());
 	DECLARE_DELEGATE_TwoParams(FCustomizeNodeDelegate, UEdGraphNode*, bool);
 	DECLARE_DELEGATE_ThreeParams(FUiSpecOverrideDelegate, FBlueprintActionContext const&, FBindingSet const&, FBlueprintActionUiSpec*);
 
diff --git a/Engine/Source/Editor/Blutility/Classes/EditorUtilityBlueprint.h b/Engine/Source/Editor/Blutility/Classes/EditorUtilityBlueprint.h
--- a/Engine/Source/Editor/Blutility/Classes/EditorUtilityBlueprint.h
+++ b/Engine/Source/Editor/Blutility/Classes/EditorUtilityBlueprint.h
@@ -11,7 +11,9 @@
 UCLASS()
 class UEditorUtilityBlueprint : public UBlueprint
 {
-	GENERATED_UCLASS_BODY()
+	GENERATED_BODY()
+public:
+	UEditorUtilityBlueprint(const FObjectInitializer& ObjectInitializer = FObjectInitializer::Get());
 
 #if WITH_EDITOR
 	// UBlueprint interface
diff --git a/Engine/Source/Editor/Blutility/Classes/EditorUtilityBlueprintFactory.h b/Engine/Source/Editor/Blutility/Classes/EditorUtilityBlueprintFactory.h
--- a/Engine/Source/Editor/Blutility/Classes/EditorUtilityBlueprintFactory.h
+++ b/Engine/Source/Editor/Blutility/Classes/EditorUtilityBlueprintFactory.h
@@ -11,7 +11,9 @@
 UCLASS()
 class UEditorUtilityBlueprintFactory : public UFactory
 {
-	GENERATED_UCLASS_BODY()
+	GENERATED_BODY()
+public:
+	UEditorUtilityBlueprintFactory(const FObjectInitializer& ObjectInitializer = FObjectInitializer::Get());
 
 	// The parent class of the created blueprint
 	UPROPERTY(EditAnywhere, Category=BlueprintFactory, meta=(AllowAbstract = "", BlueprintBaseOnly = ""))
diff --git a/Engine/Source/Editor/Blutility/Classes/GlobalEditorUtilityBase.h b/Engine/Source/Editor/Blutility/Classes/GlobalEditorUtilityBase.h
--- a/Engine/Source/Editor/Blutility/Classes/GlobalEditorUtilityBase.h
+++ b/Engine/Source/Editor/Blutility/Classes/GlobalEditorUtilityBase.h
@@ -16,7 +16,9 @@ DECLARE_DYNAMIC_MULTICAST_DELEGATE_TwoParams( FForEachAssetIteratorSignature, cl
 UCLASS(Abstract, hideCategories=(Object), Blueprintable)
 class BLUTILITY_API UGlobalEditorUtilityBase : public UObject
 {
-	GENERATED_UCLASS_BODY()
+	GENERATED_BODY()
+public:
+	UGlobalEditorUtilityBase(const FObjectInitializer& ObjectInitializer = FObjectInitializer::Get());
 
 	UPROPERTY(Category=Config, EditDefaultsOnly, BlueprintReadWrite, AssetRegistrySearchable)
 	FString HelpText;
diff --git a/Engine/Source/Editor/Blutility/Classes/PlacedEditorUtilityBase.h b/Engine/Source/Editor/Blutility/Classes/PlacedEditorUtilityBase.h
--- a/Engine/Source/Editor/Blutility/Classes/PlacedEditorUtilityBase.h
+++ b/Engine/Source/Editor/Blutility/Classes/PlacedEditorUtilityBase.h
@@ -11,7 +11,9 @@
 UCLASS(Abstract, hideCategories=(Object, Actor)/*, Blueprintable*/)
 class BLUTILITY_API APlacedEditorUtilityBase : public AActor
 {
-	GENERATED_UCLASS_BODY()
+	GENERATED_BODY()
+public:
+	APlacedEditorUtilityBase(const FObjectInitializer& ObjectInitializer = FObjectInitializer::Get());
 
 	UPROPERTY(Category=Config, EditDefaultsOnly, BlueprintReadWrite)
 	FString HelpText;
diff --git a/Engine/Source/Editor/Cascade/Classes/CascadeConfiguration.h b/Engine/Source/Editor/Cascade/Classes/CascadeConfiguration.h
--- a/Engine/Source/Editor/Cascade/Classes/CascadeConfiguration.h
+++ b/Engine/Source/Editor/Cascade/Classes/CascadeConfiguration.h
@@ -13,7 +13,7 @@
 USTRUCT()
 struct FModuleMenuMapper
 {
-	GENERATED_USTRUCT_BODY()
+	GENERATED_BODY()
 
 	UPROPERTY()
 	FString ObjName;
@@ -26,7 +26,9 @@ struct FModuleMenuMapper
 UCLASS(hidecategories=Object, config=Editor)
 class UCascadeConfiguration : public UObject
 {
-	GENERATED_UCLASS_BODY()
+	GENERATED_BODY()
+public:
+	UCascadeConfiguration(const FObjectInitializer& ObjectInitializer = FObjectInitializer::Get());
 
 	/**
 	 *	TypeData-to-base module mappings.
diff --git a/Engine/Source/Editor/Cascade/Classes/CascadeParticleSystemComponent.h b/Engine/Source/Editor/Cascade/Classes/CascadeParticleSystemComponent.h
--- a/Engine/Source/Editor/Cascade/Classes/CascadeParticleSystemComponent.h
+++ b/Engine/Source/Editor/Cascade/Classes/CascadeParticleSystemComponent.h
@@ -11,7 +11,9 @@
 UCLASS(hidecategories=Object, hidecategories=Physics, hidecategories=Collision, editinlinenew)
 class UCascadeParticleSystemComponent : public UParticleSystemComponent
 {
-	GENERATED_UCLASS_BODY()
+	GENERATED_BODY()
+public:
+	UCascadeParticleSystemComponent(const FObjectInitializer& ObjectInitializer = FObjectInitializer::Get());
 
 
 public:	
diff --git a/Engine/Source/Editor/DestructibleMeshEditor/Classes/DestructibleChunkParamsProxy.h b/Engine/Source/Editor/DestructibleMeshEditor/Classes/DestructibleChunkParamsProxy.h
--- a/Engine/Source/Editor/DestructibleMeshEditor/Classes/DestructibleChunkParamsProxy.h
+++ b/Engine/Source/Editor/DestructibleMeshEditor/Classes/DestructibleChunkParamsProxy.h
@@ -9,7 +9,9 @@
 UCLASS()
 class UDestructibleChunkParamsProxy : public UObject
 {
-	GENERATED_UCLASS_BODY()
+	GENERATED_BODY()
+public:
+	UDestructibleChunkParamsProxy(const FObjectInitializer& ObjectInitializer = FObjectInitializer::Get());
 
 	UPROPERTY()
 	UDestructibleMesh* DestructibleMesh;
@@ -27,4 +29,4 @@ class UDestructibleChunkParamsProxy : public UObject
 
 	TWeakPtr<IDestructibleMeshEditor> DestructibleMeshEditorPtr;
 #endif
-};
\ No newline at end of file
+};
diff --git a/Engine/Source/Editor/EditorLiveStreaming/EditorLiveStreamingSettings.h b/Engine/Source/Editor/EditorLiveStreaming/EditorLiveStreamingSettings.h
--- a/Engine/Source/Editor/EditorLiveStreaming/EditorLiveStreamingSettings.h
+++ b/Engine/Source/Editor/EditorLiveStreaming/EditorLiveStreamingSettings.h
@@ -38,7 +38,9 @@ namespace EEditorLiveStreamingWebCamResolution
 UCLASS(config=EditorGameAgnostic)
 class UEditorLiveStreamingSettings : public UObject
 {
-	GENERATED_UCLASS_BODY()
+	GENERATED_BODY()
+public:
+	UEditorLiveStreamingSettings(const FObjectInitializer& ObjectInitializer = FObjectInitializer::Get());
 
 public:
 
diff --git a/Engine/Source/Editor/EditorStyle/Public/Classes/EditorStyleSettings.h b/Engine/Source/Editor/EditorStyle/Public/Classes/EditorStyleSettings.h
--- a/Engine/Source/Editor/EditorStyle/Public/Classes/EditorStyleSettings.h
+++ b/Engine/Source/Editor/EditorStyle/Public/Classes/EditorStyleSettings.h
@@ -30,9 +30,11 @@ UCLASS(config=EditorUserSettings)
 class EDITORSTYLE_API UEditorStyleSettings
 	: public UObject
 {
+	GENERATED_BODY()
+public:
+	UEditorStyleSettings(const FObjectInitializer& ObjectInitializer = FObjectInitializer::Get());
 public:
 
-	GENERATED_UCLASS_BODY()
 
 public:
 
diff --git a/Engine/Source/Editor/EnvironmentQueryEditor/Classes/EdGraphSchema_EnvironmentQuery.h b/Engine/Source/Editor/EnvironmentQueryEditor/Classes/EdGraphSchema_EnvironmentQuery.h
--- a/Engine/Source/Editor/EnvironmentQueryEditor/Classes/EdGraphSchema_EnvironmentQuery.h
+++ b/Engine/Source/Editor/EnvironmentQueryEditor/Classes/EdGraphSchema_EnvironmentQuery.h
@@ -8,7 +8,9 @@
 UCLASS(MinimalAPI)
 class UEdGraphSchema_EnvironmentQuery : public UAIGraphSchema
 {
-	GENERATED_UCLASS_BODY()
+	GENERATED_BODY()
+public:
+	ENVIRONMENTQUERYEDITOR_API UEdGraphSchema_EnvironmentQuery(const FObjectInitializer& ObjectInitializer = FObjectInitializer::Get());
 
 	// Begin EdGraphSchema interface
 	virtual void CreateDefaultNodesForGraph(UEdGraph& Graph) const override;
diff --git a/Engine/Source/Editor/EnvironmentQueryEditor/Classes/EnvironmentQueryFactory.h b/Engine/Source/Editor/EnvironmentQueryEditor/Classes/EnvironmentQueryFactory.h
--- a/Engine/Source/Editor/EnvironmentQueryEditor/Classes/EnvironmentQueryFactory.h
+++ b/Engine/Source/Editor/EnvironmentQueryEditor/Classes/EnvironmentQueryFactory.h
@@ -6,7 +6,9 @@
 UCLASS(hidecategories=Object)
 class UEnvironmentQueryFactory : public UFactory
 {
-	GENERATED_UCLASS_BODY()
+	GENERATED_BODY()
+public:
+	UEnvironmentQueryFactory(const FObjectInitializer& ObjectInitializer = FObjectInitializer::Get());
 
 	// UFactory interface
 	virtual UObject* FactoryCreateNew(UClass* Class,UObject* InParent,FName Name,EObjectFlags Flags,UObject* Context,FFeedbackContext* Warn) override;
diff --git a/Engine/Source/Editor/EnvironmentQueryEditor/Classes/EnvironmentQueryGraph.h b/Engine/Source/Editor/EnvironmentQueryEditor/Classes/EnvironmentQueryGraph.h
--- a/Engine/Source/Editor/EnvironmentQueryEditor/Classes/EnvironmentQueryGraph.h
+++ b/Engine/Source/Editor/EnvironmentQueryEditor/Classes/EnvironmentQueryGraph.h
@@ -11,7 +11,9 @@ class UEnvironmentQueryGraphNode;
 UCLASS()
 class UEnvironmentQueryGraph : public UAIGraph
 {
-	GENERATED_UCLASS_BODY()
+	GENERATED_BODY()
+public:
+	UEnvironmentQueryGraph(const FObjectInitializer& ObjectInitializer = FObjectInitializer::Get());
 
 	virtual void Initialize() override;
 	virtual void OnLoaded() override;
diff --git a/Engine/Source/Editor/EnvironmentQueryEditor/Classes/EnvironmentQueryGraphNode.h b/Engine/Source/Editor/EnvironmentQueryEditor/Classes/EnvironmentQueryGraphNode.h
--- a/Engine/Source/Editor/EnvironmentQueryEditor/Classes/EnvironmentQueryGraphNode.h
+++ b/Engine/Source/Editor/EnvironmentQueryEditor/Classes/EnvironmentQueryGraphNode.h
@@ -8,7 +8,9 @@
 UCLASS()
 class UEnvironmentQueryGraphNode : public UAIGraphNode
 {
-	GENERATED_UCLASS_BODY()
+	GENERATED_BODY()
+public:
+	UEnvironmentQueryGraphNode(const FObjectInitializer& ObjectInitializer = FObjectInitializer::Get());
 
 	virtual class UEnvironmentQueryGraph* GetEnvironmentQueryGraph();
 
diff --git a/Engine/Source/Editor/EnvironmentQueryEditor/Classes/EnvironmentQueryGraphNode_Option.h b/Engine/Source/Editor/EnvironmentQueryEditor/Classes/EnvironmentQueryGraphNode_Option.h
--- a/Engine/Source/Editor/EnvironmentQueryEditor/Classes/EnvironmentQueryGraphNode_Option.h
+++ b/Engine/Source/Editor/EnvironmentQueryEditor/Classes/EnvironmentQueryGraphNode_Option.h
@@ -6,7 +6,9 @@
 UCLASS()
 class UEnvironmentQueryGraphNode_Option : public UEnvironmentQueryGraphNode
 {
-	GENERATED_UCLASS_BODY()
+	GENERATED_BODY()
+public:
+	UEnvironmentQueryGraphNode_Option(const FObjectInitializer& ObjectInitializer = FObjectInitializer::Get());
 
 	virtual void AllocateDefaultPins() override;
 	virtual void PostPlacedNewNode() override;
diff --git a/Engine/Source/Editor/EnvironmentQueryEditor/Classes/EnvironmentQueryGraphNode_Root.h b/Engine/Source/Editor/EnvironmentQueryEditor/Classes/EnvironmentQueryGraphNode_Root.h
--- a/Engine/Source/Editor/EnvironmentQueryEditor/Classes/EnvironmentQueryGraphNode_Root.h
+++ b/Engine/Source/Editor/EnvironmentQueryEditor/Classes/EnvironmentQueryGraphNode_Root.h
@@ -5,7 +5,9 @@
 UCLASS()
 class UEnvironmentQueryGraphNode_Root : public UEnvironmentQueryGraphNode
 {
-	GENERATED_UCLASS_BODY()
+	GENERATED_BODY()
+public:
+	UEnvironmentQueryGraphNode_Root(const FObjectInitializer& ObjectInitializer = FObjectInitializer::Get());
 
 	virtual void AllocateDefaultPins() override;
 	virtual FText GetNodeTitle(ENodeTitleType::Type TitleType) const override;
diff --git a/Engine/Source/Editor/EnvironmentQueryEditor/Classes/EnvironmentQueryGraphNode_Test.h b/Engine/Source/Editor/EnvironmentQueryEditor/Classes/EnvironmentQueryGraphNode_Test.h
--- a/Engine/Source/Editor/EnvironmentQueryEditor/Classes/EnvironmentQueryGraphNode_Test.h
+++ b/Engine/Source/Editor/EnvironmentQueryEditor/Classes/EnvironmentQueryGraphNode_Test.h
@@ -6,7 +6,9 @@
 UCLASS()
 class UEnvironmentQueryGraphNode_Test : public UEnvironmentQueryGraphNode
 {
-	GENERATED_UCLASS_BODY()
+	GENERATED_BODY()
+public:
+	UEnvironmentQueryGraphNode_Test(const FObjectInitializer& ObjectInitializer = FObjectInitializer::Get());
 
 	/** weight percent for display */
 	UPROPERTY()
diff --git a/Engine/Source/Editor/FoliageEdit/Private/ActorFactoryProceduralFoliage.h b/Engine/Source/Editor/FoliageEdit/Private/ActorFactoryProceduralFoliage.h
--- a/Engine/Source/Editor/FoliageEdit/Private/ActorFactoryProceduralFoliage.h
+++ b/Engine/Source/Editor/FoliageEdit/Private/ActorFactoryProceduralFoliage.h
@@ -8,7 +8,9 @@
 UCLASS(MinimalAPI, config=Editor)
 class UActorFactoryProceduralFoliage : public UActorFactoryBoxVolume
 {
-	GENERATED_UCLASS_BODY()
+	GENERATED_BODY()
+public:
+	FOLIAGEEDIT_API UActorFactoryProceduralFoliage(const FObjectInitializer& ObjectInitializer = FObjectInitializer::Get());
 
 	// Begin UActorFactory Interface
 	virtual bool CanCreateActorFrom( const FAssetData& AssetData, FText& OutErrorMsg ) override;
diff --git a/Engine/Source/Editor/FoliageEdit/Private/LandscapeGrassTypeFactory.h b/Engine/Source/Editor/FoliageEdit/Private/LandscapeGrassTypeFactory.h
--- a/Engine/Source/Editor/FoliageEdit/Private/LandscapeGrassTypeFactory.h
+++ b/Engine/Source/Editor/FoliageEdit/Private/LandscapeGrassTypeFactory.h
@@ -12,7 +12,9 @@
 UCLASS()
 class ULandscapeGrassTypeFactory : public UFactory
 {
-	GENERATED_UCLASS_BODY()
+	GENERATED_BODY()
+public:
+	ULandscapeGrassTypeFactory(const FObjectInitializer& ObjectInitializer = FObjectInitializer::Get());
 
 	// UFactory interface
 	virtual UObject* FactoryCreateNew(UClass* Class, UObject* InParent, FName Name, EObjectFlags Flags, UObject* Context, FFeedbackContext* Warn) override;
diff --git a/Engine/Source/Editor/FoliageEdit/Private/ProceduralFoliageFactory.h b/Engine/Source/Editor/FoliageEdit/Private/ProceduralFoliageFactory.h
--- a/Engine/Source/Editor/FoliageEdit/Private/ProceduralFoliageFactory.h
+++ b/Engine/Source/Editor/FoliageEdit/Private/ProceduralFoliageFactory.h
@@ -12,7 +12,9 @@
 UCLASS()
 class UProceduralFoliageFactory : public UFactory
 {
-	GENERATED_UCLASS_BODY()
+	GENERATED_BODY()
+public:
+	UProceduralFoliageFactory(const FObjectInitializer& ObjectInitializer = FObjectInitializer::Get());
 
 	// UFactory interface
 	virtual UObject* FactoryCreateNew(UClass* Class, UObject* InParent, FName Name, EObjectFlags Flags, UObject* Context, FFeedbackContext* Warn) override;
diff --git a/Engine/Source/Editor/GameProjectGeneration/Classes/DefaultTemplateProjectDefs.h b/Engine/Source/Editor/GameProjectGeneration/Classes/DefaultTemplateProjectDefs.h
--- a/Engine/Source/Editor/GameProjectGeneration/Classes/DefaultTemplateProjectDefs.h
+++ b/Engine/Source/Editor/GameProjectGeneration/Classes/DefaultTemplateProjectDefs.h
@@ -8,7 +8,9 @@
 UCLASS()
 class GAMEPROJECTGENERATION_API UDefaultTemplateProjectDefs : public UTemplateProjectDefs
 {
-	GENERATED_UCLASS_BODY()
+	GENERATED_BODY()
+public:
+	UDefaultTemplateProjectDefs(const FObjectInitializer& ObjectInitializer = FObjectInitializer::Get());
 
 	virtual bool GeneratesCode(const FString& ProjectTemplatePath) const override;
 
diff --git a/Engine/Source/Editor/GameProjectGeneration/Classes/TemplateProjectDefs.h b/Engine/Source/Editor/GameProjectGeneration/Classes/TemplateProjectDefs.h
--- a/Engine/Source/Editor/GameProjectGeneration/Classes/TemplateProjectDefs.h
+++ b/Engine/Source/Editor/GameProjectGeneration/Classes/TemplateProjectDefs.h
@@ -19,7 +19,7 @@ struct FTemplateConfigValue
 USTRUCT()
 struct FTemplateReplacement
 {
-	GENERATED_USTRUCT_BODY()
+	GENERATED_BODY()
 
 	UPROPERTY()
 	TArray<FString> Extensions;
@@ -37,7 +37,7 @@ struct FTemplateReplacement
 USTRUCT()
 struct FTemplateFolderRename
 {
-	GENERATED_USTRUCT_BODY()
+	GENERATED_BODY()
 
 	UPROPERTY()
 	FString From;
@@ -49,7 +49,7 @@ struct FTemplateFolderRename
 USTRUCT()
 struct FLocalizedTemplateString
 {
-	GENERATED_USTRUCT_BODY()
+	GENERATED_BODY()
 
 	UPROPERTY()
 	FString Language;
@@ -61,7 +61,9 @@ struct FLocalizedTemplateString
 UCLASS(abstract,config=TemplateDefs,MinimalAPI)
 class UTemplateProjectDefs : public UObject
 {
-	GENERATED_UCLASS_BODY()
+	GENERATED_BODY()
+public:
+	GAMEPROJECTGENERATION_API UTemplateProjectDefs(const FObjectInitializer& ObjectInitializer = FObjectInitializer::Get());
 
 	UPROPERTY(config)
 	TArray<FLocalizedTemplateString> LocalizedDisplayNames;
diff --git a/Engine/Source/Editor/GameplayAbilitiesEditor/Classes/GameplayAbilitiesBlueprintFactory.h b/Engine/Source/Editor/GameplayAbilitiesEditor/Classes/GameplayAbilitiesBlueprintFactory.h
--- a/Engine/Source/Editor/GameplayAbilitiesEditor/Classes/GameplayAbilitiesBlueprintFactory.h
+++ b/Engine/Source/Editor/GameplayAbilitiesEditor/Classes/GameplayAbilitiesBlueprintFactory.h
@@ -7,7 +7,9 @@
 UCLASS(HideCategories=Object, MinimalAPI)
 class UGameplayAbilitiesBlueprintFactory : public UFactory
 {
-	GENERATED_UCLASS_BODY()
+	GENERATED_BODY()
+public:
+	GAMEPLAYABILITIESEDITOR_API UGameplayAbilitiesBlueprintFactory(const FObjectInitializer& ObjectInitializer = FObjectInitializer::Get());
 
 	// The type of blueprint that will be created
 	UPROPERTY(EditAnywhere, Category=GameplayAbilitiesBlueprintFactory)
diff --git a/Engine/Source/Editor/GameplayAbilitiesEditor/Classes/GameplayAbilityGraph.h b/Engine/Source/Editor/GameplayAbilitiesEditor/Classes/GameplayAbilityGraph.h
--- a/Engine/Source/Editor/GameplayAbilitiesEditor/Classes/GameplayAbilityGraph.h
+++ b/Engine/Source/Editor/GameplayAbilitiesEditor/Classes/GameplayAbilityGraph.h
@@ -7,6 +7,8 @@
 UCLASS(MinimalAPI)
 class UGameplayAbilityGraph : public UEdGraph
 {
-	GENERATED_UCLASS_BODY()
+	GENERATED_BODY()
+public:
+	GAMEPLAYABILITIESEDITOR_API UGameplayAbilityGraph(const FObjectInitializer& ObjectInitializer = FObjectInitializer::Get());
 };
 
diff --git a/Engine/Source/Editor/GameplayAbilitiesEditor/Classes/GameplayAbilityGraphSchema.h b/Engine/Source/Editor/GameplayAbilitiesEditor/Classes/GameplayAbilityGraphSchema.h
--- a/Engine/Source/Editor/GameplayAbilitiesEditor/Classes/GameplayAbilityGraphSchema.h
+++ b/Engine/Source/Editor/GameplayAbilitiesEditor/Classes/GameplayAbilityGraphSchema.h
@@ -8,7 +8,9 @@
 UCLASS(MinimalAPI)
 class UGameplayAbilityGraphSchema : public UEdGraphSchema_K2
 {
-	GENERATED_UCLASS_BODY()
+	GENERATED_BODY()
+public:
+	GAMEPLAYABILITIESEDITOR_API UGameplayAbilityGraphSchema(const FObjectInitializer& ObjectInitializer = FObjectInitializer::Get());
 
 	/**
 	* Creates a new variable getter node and adds it to ParentGraph
@@ -33,4 +35,4 @@ class UGameplayAbilityGraphSchema : public UEdGraphSchema_K2
 	virtual class UK2Node_VariableSet* SpawnVariableSetNode(const FVector2D GraphPosition, class UEdGraph* ParentGraph, FName VariableName, UStruct* Source) const override;
 
 	virtual bool ShouldAlwaysPurgeOnModification() const override { return true; }
-};
\ No newline at end of file
+};
diff --git a/Engine/Source/Editor/GameplayAbilitiesEditor/Classes/K2Node_GameplayCueEvent.h b/Engine/Source/Editor/GameplayAbilitiesEditor/Classes/K2Node_GameplayCueEvent.h
--- a/Engine/Source/Editor/GameplayAbilitiesEditor/Classes/K2Node_GameplayCueEvent.h
+++ b/Engine/Source/Editor/GameplayAbilitiesEditor/Classes/K2Node_GameplayCueEvent.h
@@ -7,7 +7,9 @@
 UCLASS()
 class UK2Node_GameplayCueEvent : public UK2Node_Event
 {
-	GENERATED_UCLASS_BODY()
+	GENERATED_BODY()
+public:
+	UK2Node_GameplayCueEvent(const FObjectInitializer& ObjectInitializer = FObjectInitializer::Get());
 
 	// UObject interface
 	virtual void Serialize(FArchive& Ar) override;
@@ -23,4 +25,4 @@ class UK2Node_GameplayCueEvent : public UK2Node_Event
 	// Begin UK2Node interface
 	virtual void GetMenuActions(FBlueprintActionDatabaseRegistrar& ActionRegistrar) const override;
 	// End UK2Node interface
-};
\ No newline at end of file
+};
diff --git a/Engine/Source/Editor/GameplayAbilitiesEditor/Classes/K2Node_GameplayEffectVariable.h b/Engine/Source/Editor/GameplayAbilitiesEditor/Classes/K2Node_GameplayEffectVariable.h
--- a/Engine/Source/Editor/GameplayAbilitiesEditor/Classes/K2Node_GameplayEffectVariable.h
+++ b/Engine/Source/Editor/GameplayAbilitiesEditor/Classes/K2Node_GameplayEffectVariable.h
@@ -11,7 +11,9 @@
 UCLASS()
 class GAMEPLAYABILITIESEDITOR_API UK2Node_GameplayEffectVariable : public UK2Node_VariableGet // only need a custom node for the getter
 {
-	GENERATED_UCLASS_BODY()
+	GENERATED_BODY()
+public:
+	UK2Node_GameplayEffectVariable(const FObjectInitializer& ObjectInitializer = FObjectInitializer::Get());
 
 public:
 	UPROPERTY()
diff --git a/Engine/Source/Editor/GameplayAbilitiesEditor/Classes/K2Node_LatentAbilityCall.h b/Engine/Source/Editor/GameplayAbilitiesEditor/Classes/K2Node_LatentAbilityCall.h
--- a/Engine/Source/Editor/GameplayAbilitiesEditor/Classes/K2Node_LatentAbilityCall.h
+++ b/Engine/Source/Editor/GameplayAbilitiesEditor/Classes/K2Node_LatentAbilityCall.h
@@ -11,7 +11,9 @@
 UCLASS()
 class UK2Node_LatentAbilityCall : public UK2Node_BaseAsyncTask
 {
-	GENERATED_UCLASS_BODY()
+	GENERATED_BODY()
+public:
+	UK2Node_LatentAbilityCall(const FObjectInitializer& ObjectInitializer = FObjectInitializer::Get());
 
 	// UEdGraphNode interface
 	virtual bool CanCreateUnderSpecifiedSchema(const UEdGraphSchema* DesiredSchema) const;
diff --git a/Engine/Source/Editor/GameplayTagsEditor/Classes/GameplayTagsK2Node_LiteralGameplayTag.h b/Engine/Source/Editor/GameplayTagsEditor/Classes/GameplayTagsK2Node_LiteralGameplayTag.h
--- a/Engine/Source/Editor/GameplayTagsEditor/Classes/GameplayTagsK2Node_LiteralGameplayTag.h
+++ b/Engine/Source/Editor/GameplayTagsEditor/Classes/GameplayTagsK2Node_LiteralGameplayTag.h
@@ -8,7 +8,9 @@
 UCLASS()
 class UGameplayTagsK2Node_LiteralGameplayTag : public UK2Node
 {
-	GENERATED_UCLASS_BODY()
+	GENERATED_BODY()
+public:
+	UGameplayTagsK2Node_LiteralGameplayTag(const FObjectInitializer& ObjectInitializer = FObjectInitializer::Get());
 
 #if WITH_EDITOR
 	// Begin UEdGraphNode interface
diff --git a/Engine/Source/Editor/GameplayTagsEditor/Classes/GameplayTagsK2Node_MultiCompareBase.h b/Engine/Source/Editor/GameplayTagsEditor/Classes/GameplayTagsK2Node_MultiCompareBase.h
--- a/Engine/Source/Editor/GameplayTagsEditor/Classes/GameplayTagsK2Node_MultiCompareBase.h
+++ b/Engine/Source/Editor/GameplayTagsEditor/Classes/GameplayTagsK2Node_MultiCompareBase.h
@@ -11,7 +11,9 @@ struct FGameplayTagContainer;
 UCLASS()
 class UGameplayTagsK2Node_MultiCompareBase : public UK2Node
 {
-	GENERATED_UCLASS_BODY()
+	GENERATED_BODY()
+public:
+	UGameplayTagsK2Node_MultiCompareBase(const FObjectInitializer& ObjectInitializer = FObjectInitializer::Get());
 
 	UPROPERTY(EditAnywhere, Category = PinOptions)
 	int32 NumberOfPins;
diff --git a/Engine/Source/Editor/GameplayTagsEditor/Classes/GameplayTagsK2Node_MultiCompareGameplayTagAssetInterface.h b/Engine/Source/Editor/GameplayTagsEditor/Classes/GameplayTagsK2Node_MultiCompareGameplayTagAssetInterface.h
--- a/Engine/Source/Editor/GameplayTagsEditor/Classes/GameplayTagsK2Node_MultiCompareGameplayTagAssetInterface.h
+++ b/Engine/Source/Editor/GameplayTagsEditor/Classes/GameplayTagsK2Node_MultiCompareGameplayTagAssetInterface.h
@@ -8,7 +8,9 @@
 UCLASS()
 class UGameplayTagsK2Node_MultiCompareGameplayTagAssetInterface : public UGameplayTagsK2Node_MultiCompareBase
 {
-	GENERATED_UCLASS_BODY()
+	GENERATED_BODY()
+public:
+	UGameplayTagsK2Node_MultiCompareGameplayTagAssetInterface(const FObjectInitializer& ObjectInitializer = FObjectInitializer::Get());
 
 	// UEdGraphNode interface
 	virtual void AllocateDefaultPins() override;
@@ -23,4 +25,4 @@ class UGameplayTagsK2Node_MultiCompareGameplayTagAssetInterface : public UGamepl
 private:
 
 	virtual void AddPinToSwitchNode() override;
-};
\ No newline at end of file
+};
diff --git a/Engine/Source/Editor/GameplayTagsEditor/Classes/GameplayTagsK2Node_MultiCompareGameplayTagAssetInterfaceSingleTags.h b/Engine/Source/Editor/GameplayTagsEditor/Classes/GameplayTagsK2Node_MultiCompareGameplayTagAssetInterfaceSingleTags.h
--- a/Engine/Source/Editor/GameplayTagsEditor/Classes/GameplayTagsK2Node_MultiCompareGameplayTagAssetInterfaceSingleTags.h
+++ b/Engine/Source/Editor/GameplayTagsEditor/Classes/GameplayTagsK2Node_MultiCompareGameplayTagAssetInterfaceSingleTags.h
@@ -8,7 +8,9 @@
 UCLASS()
 class UGameplayTagsK2Node_MultiCompareGameplayTagAssetInterfaceSingleTags : public UGameplayTagsK2Node_MultiCompareBase
 {
-	GENERATED_UCLASS_BODY()
+	GENERATED_BODY()
+public:
+	UGameplayTagsK2Node_MultiCompareGameplayTagAssetInterfaceSingleTags(const FObjectInitializer& ObjectInitializer = FObjectInitializer::Get());
 
 	// UEdGraphNode interface
 	virtual void AllocateDefaultPins() override;
diff --git a/Engine/Source/Editor/GameplayTagsEditor/Classes/GameplayTagsK2Node_MultiCompareGameplayTagContainer.h b/Engine/Source/Editor/GameplayTagsEditor/Classes/GameplayTagsK2Node_MultiCompareGameplayTagContainer.h
--- a/Engine/Source/Editor/GameplayTagsEditor/Classes/GameplayTagsK2Node_MultiCompareGameplayTagContainer.h
+++ b/Engine/Source/Editor/GameplayTagsEditor/Classes/GameplayTagsK2Node_MultiCompareGameplayTagContainer.h
@@ -8,7 +8,9 @@
 UCLASS()
 class UGameplayTagsK2Node_MultiCompareGameplayTagContainer : public UGameplayTagsK2Node_MultiCompareBase
 {
-	GENERATED_UCLASS_BODY()
+	GENERATED_BODY()
+public:
+	UGameplayTagsK2Node_MultiCompareGameplayTagContainer(const FObjectInitializer& ObjectInitializer = FObjectInitializer::Get());
 
 	// UEdGraphNode interface
 	virtual void AllocateDefaultPins() override;
diff --git a/Engine/Source/Editor/GameplayTagsEditor/Classes/GameplayTagsK2Node_MultiCompareGameplayTagContainerSingleTags.h b/Engine/Source/Editor/GameplayTagsEditor/Classes/GameplayTagsK2Node_MultiCompareGameplayTagContainerSingleTags.h
--- a/Engine/Source/Editor/GameplayTagsEditor/Classes/GameplayTagsK2Node_MultiCompareGameplayTagContainerSingleTags.h
+++ b/Engine/Source/Editor/GameplayTagsEditor/Classes/GameplayTagsK2Node_MultiCompareGameplayTagContainerSingleTags.h
@@ -8,7 +8,9 @@
 UCLASS()
 class UGameplayTagsK2Node_MultiCompareGameplayTagContainerSingleTags : public UGameplayTagsK2Node_MultiCompareBase
 {
-	GENERATED_UCLASS_BODY()
+	GENERATED_BODY()
+public:
+	UGameplayTagsK2Node_MultiCompareGameplayTagContainerSingleTags(const FObjectInitializer& ObjectInitializer = FObjectInitializer::Get());
 
 	// UEdGraphNode interface
 	virtual void AllocateDefaultPins() override;
diff --git a/Engine/Source/Editor/GameplayTagsEditor/Classes/GameplayTagsK2Node_SwitchGameplayTag.h b/Engine/Source/Editor/GameplayTagsEditor/Classes/GameplayTagsK2Node_SwitchGameplayTag.h
--- a/Engine/Source/Editor/GameplayTagsEditor/Classes/GameplayTagsK2Node_SwitchGameplayTag.h
+++ b/Engine/Source/Editor/GameplayTagsEditor/Classes/GameplayTagsK2Node_SwitchGameplayTag.h
@@ -11,7 +11,9 @@ struct FGameplayTag;
 UCLASS(MinimalAPI)
 class UGameplayTagsK2Node_SwitchGameplayTag : public UK2Node_Switch
 {
-	GENERATED_UCLASS_BODY()
+	GENERATED_BODY()
+public:
+	GAMEPLAYTAGSEDITOR_API UGameplayTagsK2Node_SwitchGameplayTag(const FObjectInitializer& ObjectInitializer = FObjectInitializer::Get());
 
 	UPROPERTY(EditAnywhere, Category = PinOptions)
 	TArray<FGameplayTag> PinTags;
diff --git a/Engine/Source/Editor/GameplayTagsEditor/Classes/GameplayTagsK2Node_SwitchGameplayTagContainer.h b/Engine/Source/Editor/GameplayTagsEditor/Classes/GameplayTagsK2Node_SwitchGameplayTagContainer.h
--- a/Engine/Source/Editor/GameplayTagsEditor/Classes/GameplayTagsK2Node_SwitchGameplayTagContainer.h
+++ b/Engine/Source/Editor/GameplayTagsEditor/Classes/GameplayTagsK2Node_SwitchGameplayTagContainer.h
@@ -11,7 +11,9 @@ struct FGameplayTagContainer;
 UCLASS(MinimalAPI)
 class UGameplayTagsK2Node_SwitchGameplayTagContainer : public UK2Node_Switch
 {
-	GENERATED_UCLASS_BODY()
+	GENERATED_BODY()
+public:
+	GAMEPLAYTAGSEDITOR_API UGameplayTagsK2Node_SwitchGameplayTagContainer(const FObjectInitializer& ObjectInitializer = FObjectInitializer::Get());
 
 	UPROPERTY(EditAnywhere, Category = PinOptions)
 	TArray<FGameplayTagContainer> PinContainers;
diff --git a/Engine/Source/Editor/GeometryMode/Classes/GeomModifier.h b/Engine/Source/Editor/GeometryMode/Classes/GeomModifier.h
--- a/Engine/Source/Editor/GeometryMode/Classes/GeomModifier.h
+++ b/Engine/Source/Editor/GeometryMode/Classes/GeomModifier.h
@@ -11,7 +11,9 @@
 UCLASS(abstract, hidecategories=(Object, GeomModifier))
 class UGeomModifier : public UObject
 {
-	GENERATED_UCLASS_BODY()
+	GENERATED_BODY()
+public:
+	UGeomModifier(const FObjectInitializer& ObjectInitializer = FObjectInitializer::Get());
 
 	/** A human readable name for this modifier (appears on buttons, menus, etc) */
 	UPROPERTY(EditAnywhere, Category=GeomModifier)
diff --git a/Engine/Source/Editor/GeometryMode/Classes/GeomModifier_Clip.h b/Engine/Source/Editor/GeometryMode/Classes/GeomModifier_Clip.h
--- a/Engine/Source/Editor/GeometryMode/Classes/GeomModifier_Clip.h
+++ b/Engine/Source/Editor/GeometryMode/Classes/GeomModifier_Clip.h
@@ -11,7 +11,9 @@
 UCLASS()
 class UGeomModifier_Clip : public UGeomModifier_Edit
 {
-	GENERATED_UCLASS_BODY()
+	GENERATED_BODY()
+public:
+	UGeomModifier_Clip(const FObjectInitializer& ObjectInitializer = FObjectInitializer::Get());
 
 	UPROPERTY(EditAnywhere, Category=Settings)
 	uint32 bFlipNormal:1;
diff --git a/Engine/Source/Editor/GeometryMode/Classes/GeomModifier_Create.h b/Engine/Source/Editor/GeometryMode/Classes/GeomModifier_Create.h
--- a/Engine/Source/Editor/GeometryMode/Classes/GeomModifier_Create.h
+++ b/Engine/Source/Editor/GeometryMode/Classes/GeomModifier_Create.h
@@ -11,7 +11,9 @@
 UCLASS()
 class UGeomModifier_Create : public UGeomModifier_Edit
 {
-	GENERATED_UCLASS_BODY()
+	GENERATED_BODY()
+public:
+	UGeomModifier_Create(const FObjectInitializer& ObjectInitializer = FObjectInitializer::Get());
 
 
 	// Begin UGeomModifier Interface
diff --git a/Engine/Source/Editor/GeometryMode/Classes/GeomModifier_Delete.h b/Engine/Source/Editor/GeometryMode/Classes/GeomModifier_Delete.h
--- a/Engine/Source/Editor/GeometryMode/Classes/GeomModifier_Delete.h
+++ b/Engine/Source/Editor/GeometryMode/Classes/GeomModifier_Delete.h
@@ -11,7 +11,9 @@
 UCLASS()
 class UGeomModifier_Delete : public UGeomModifier_Edit
 {
-	GENERATED_UCLASS_BODY()
+	GENERATED_BODY()
+public:
+	UGeomModifier_Delete(const FObjectInitializer& ObjectInitializer = FObjectInitializer::Get());
 
 
 	// Begin UGeomModifier Interface
diff --git a/Engine/Source/Editor/GeometryMode/Classes/GeomModifier_Edit.h b/Engine/Source/Editor/GeometryMode/Classes/GeomModifier_Edit.h
--- a/Engine/Source/Editor/GeometryMode/Classes/GeomModifier_Edit.h
+++ b/Engine/Source/Editor/GeometryMode/Classes/GeomModifier_Edit.h
@@ -16,7 +16,9 @@
 UCLASS(autoexpandcategories=Settings)
 class UGeomModifier_Edit : public UGeomModifier
 {
-	GENERATED_UCLASS_BODY()
+	GENERATED_BODY()
+public:
+	UGeomModifier_Edit(const FObjectInitializer& ObjectInitializer = FObjectInitializer::Get());
 
 
 
diff --git a/Engine/Source/Editor/GeometryMode/Classes/GeomModifier_Extrude.h b/Engine/Source/Editor/GeometryMode/Classes/GeomModifier_Extrude.h
--- a/Engine/Source/Editor/GeometryMode/Classes/GeomModifier_Extrude.h
+++ b/Engine/Source/Editor/GeometryMode/Classes/GeomModifier_Extrude.h
@@ -11,7 +11,9 @@
 UCLASS()
 class UGeomModifier_Extrude : public UGeomModifier_Edit
 {
-	GENERATED_UCLASS_BODY()
+	GENERATED_BODY()
+public:
+	UGeomModifier_Extrude(const FObjectInitializer& ObjectInitializer = FObjectInitializer::Get());
 
 	UPROPERTY(EditAnywhere, Category=Settings)
 	int32 Length;
diff --git a/Engine/Source/Editor/GeometryMode/Classes/GeomModifier_Flip.h b/Engine/Source/Editor/GeometryMode/Classes/GeomModifier_Flip.h
--- a/Engine/Source/Editor/GeometryMode/Classes/GeomModifier_Flip.h
+++ b/Engine/Source/Editor/GeometryMode/Classes/GeomModifier_Flip.h
@@ -11,7 +11,9 @@
 UCLASS()
 class UGeomModifier_Flip : public UGeomModifier_Edit
 {
-	GENERATED_UCLASS_BODY()
+	GENERATED_BODY()
+public:
+	UGeomModifier_Flip(const FObjectInitializer& ObjectInitializer = FObjectInitializer::Get());
 
 
 	// Begin UGeomModifier Interface
diff --git a/Engine/Source/Editor/GeometryMode/Classes/GeomModifier_Lathe.h b/Engine/Source/Editor/GeometryMode/Classes/GeomModifier_Lathe.h
--- a/Engine/Source/Editor/GeometryMode/Classes/GeomModifier_Lathe.h
+++ b/Engine/Source/Editor/GeometryMode/Classes/GeomModifier_Lathe.h
@@ -11,7 +11,9 @@
 UCLASS()
 class UGeomModifier_Lathe : public UGeomModifier_Edit
 {
-	GENERATED_UCLASS_BODY()
+	GENERATED_BODY()
+public:
+	UGeomModifier_Lathe(const FObjectInitializer& ObjectInitializer = FObjectInitializer::Get());
 
 	UPROPERTY(EditAnywhere, Category=Settings)
 	int32 TotalSegments;
diff --git a/Engine/Source/Editor/GeometryMode/Classes/GeomModifier_Optimize.h b/Engine/Source/Editor/GeometryMode/Classes/GeomModifier_Optimize.h
--- a/Engine/Source/Editor/GeometryMode/Classes/GeomModifier_Optimize.h
+++ b/Engine/Source/Editor/GeometryMode/Classes/GeomModifier_Optimize.h
@@ -11,7 +11,9 @@
 UCLASS()
 class UGeomModifier_Optimize : public UGeomModifier_Edit
 {
-	GENERATED_UCLASS_BODY()
+	GENERATED_BODY()
+public:
+	UGeomModifier_Optimize(const FObjectInitializer& ObjectInitializer = FObjectInitializer::Get());
 
 
 	// Begin UGeomModifier Interface
diff --git a/Engine/Source/Editor/GeometryMode/Classes/GeomModifier_Pen.h b/Engine/Source/Editor/GeometryMode/Classes/GeomModifier_Pen.h
--- a/Engine/Source/Editor/GeometryMode/Classes/GeomModifier_Pen.h
+++ b/Engine/Source/Editor/GeometryMode/Classes/GeomModifier_Pen.h
@@ -11,7 +11,9 @@
 UCLASS()
 class UGeomModifier_Pen : public UGeomModifier_Edit
 {
-	GENERATED_UCLASS_BODY()
+	GENERATED_BODY()
+public:
+	UGeomModifier_Pen(const FObjectInitializer& ObjectInitializer = FObjectInitializer::Get());
 
 	/** If true, the shape will be automatically extruded into a brush upon completion. */
 	UPROPERTY(EditAnywhere, Category=Settings)
diff --git a/Engine/Source/Editor/GeometryMode/Classes/GeomModifier_Split.h b/Engine/Source/Editor/GeometryMode/Classes/GeomModifier_Split.h
--- a/Engine/Source/Editor/GeometryMode/Classes/GeomModifier_Split.h
+++ b/Engine/Source/Editor/GeometryMode/Classes/GeomModifier_Split.h
@@ -11,7 +11,9 @@
 UCLASS()
 class UGeomModifier_Split : public UGeomModifier_Edit
 {
-	GENERATED_UCLASS_BODY()
+	GENERATED_BODY()
+public:
+	UGeomModifier_Split(const FObjectInitializer& ObjectInitializer = FObjectInitializer::Get());
 
 
 	// Begin UGeomModifier Interface
diff --git a/Engine/Source/Editor/GeometryMode/Classes/GeomModifier_Triangulate.h b/Engine/Source/Editor/GeometryMode/Classes/GeomModifier_Triangulate.h
--- a/Engine/Source/Editor/GeometryMode/Classes/GeomModifier_Triangulate.h
+++ b/Engine/Source/Editor/GeometryMode/Classes/GeomModifier_Triangulate.h
@@ -11,7 +11,9 @@
 UCLASS()
 class UGeomModifier_Triangulate : public UGeomModifier_Edit
 {
-	GENERATED_UCLASS_BODY()
+	GENERATED_BODY()
+public:
+	UGeomModifier_Triangulate(const FObjectInitializer& ObjectInitializer = FObjectInitializer::Get());
 
 
 	// Begin UGeomModifier Interface
diff --git a/Engine/Source/Editor/GeometryMode/Classes/GeomModifier_Turn.h b/Engine/Source/Editor/GeometryMode/Classes/GeomModifier_Turn.h
--- a/Engine/Source/Editor/GeometryMode/Classes/GeomModifier_Turn.h
+++ b/Engine/Source/Editor/GeometryMode/Classes/GeomModifier_Turn.h
@@ -11,7 +11,9 @@
 UCLASS()
 class UGeomModifier_Turn : public UGeomModifier_Edit
 {
-	GENERATED_UCLASS_BODY()
+	GENERATED_BODY()
+public:
+	UGeomModifier_Turn(const FObjectInitializer& ObjectInitializer = FObjectInitializer::Get());
 
 
 	// Begin UGeomModifier Interface
diff --git a/Engine/Source/Editor/GeometryMode/Classes/GeomModifier_Weld.h b/Engine/Source/Editor/GeometryMode/Classes/GeomModifier_Weld.h
--- a/Engine/Source/Editor/GeometryMode/Classes/GeomModifier_Weld.h
+++ b/Engine/Source/Editor/GeometryMode/Classes/GeomModifier_Weld.h
@@ -10,7 +10,9 @@
 UCLASS()
 class UGeomModifier_Weld : public UGeomModifier_Edit
 {
-	GENERATED_UCLASS_BODY()
+	GENERATED_BODY()
+public:
+	UGeomModifier_Weld(const FObjectInitializer& ObjectInitializer = FObjectInitializer::Get());
 
 
 	// Begin UGeomModifier Interface
diff --git a/Engine/Source/Editor/GraphEditor/Public/GraphEditorSettings.h b/Engine/Source/Editor/GraphEditor/Public/GraphEditorSettings.h
--- a/Engine/Source/Editor/GraphEditor/Public/GraphEditorSettings.h
+++ b/Engine/Source/Editor/GraphEditor/Public/GraphEditorSettings.h
@@ -13,7 +13,9 @@ UCLASS(config=EditorUserSettings)
 class GRAPHEDITOR_API UGraphEditorSettings
 	:	public UObject
 {
-	GENERATED_UCLASS_BODY()
+	GENERATED_BODY()
+public:
+	UGraphEditorSettings(const FObjectInitializer& ObjectInitializer = FObjectInitializer::Get());
 
 public:
 
diff --git a/Engine/Source/Editor/HardwareTargeting/Public/HardwareTargetingSettings.h b/Engine/Source/Editor/HardwareTargeting/Public/HardwareTargetingSettings.h
--- a/Engine/Source/Editor/HardwareTargeting/Public/HardwareTargetingSettings.h
+++ b/Engine/Source/Editor/HardwareTargeting/Public/HardwareTargetingSettings.h
@@ -42,7 +42,9 @@ namespace EGraphicsPreset
 UCLASS(config=Engine, defaultconfig)
 class HARDWARETARGETING_API UHardwareTargetingSettings : public UObject
 {
-	GENERATED_UCLASS_BODY()
+	GENERATED_BODY()
+public:
+	UHardwareTargetingSettings(const FObjectInitializer& ObjectInitializer = FObjectInitializer::Get());
 
 	/** Enum specifying which class of hardware this game is targeting */
 	UPROPERTY(config, EditAnywhere, category=None)
@@ -77,4 +79,4 @@ private:
 	/** Holds an event delegate that is executed when a setting has changed. */
 	FSettingChangedEvent SettingChangedEvent;
 #endif
-};
\ No newline at end of file
+};
diff --git a/Engine/Source/Editor/IntroTutorials/Private/EditorTutorialFactory.h b/Engine/Source/Editor/IntroTutorials/Private/EditorTutorialFactory.h
--- a/Engine/Source/Editor/IntroTutorials/Private/EditorTutorialFactory.h
+++ b/Engine/Source/Editor/IntroTutorials/Private/EditorTutorialFactory.h
@@ -7,7 +7,9 @@
 UCLASS(hidecategories=Object, collapsecategories)
 class UEditorTutorialFactory : public UFactory
 {
-	GENERATED_UCLASS_BODY()
+	GENERATED_BODY()
+public:
+	UEditorTutorialFactory(const FObjectInitializer& ObjectInitializer = FObjectInitializer::Get());
 
 	// UFactory Interface
 	virtual UObject* FactoryCreateNew(UClass* Class, UObject* InParent, FName Name, EObjectFlags Flags, UObject* Context, FFeedbackContext* Warn) override;
diff --git a/Engine/Source/Editor/IntroTutorials/Private/EditorTutorialImportFactory.h b/Engine/Source/Editor/IntroTutorials/Private/EditorTutorialImportFactory.h
--- a/Engine/Source/Editor/IntroTutorials/Private/EditorTutorialImportFactory.h
+++ b/Engine/Source/Editor/IntroTutorials/Private/EditorTutorialImportFactory.h
@@ -7,7 +7,9 @@
 UCLASS(hidecategories=Object, collapsecategories)
 class UEditorTutorialImportFactory : public UFactory, public FReimportHandler
 {
-	GENERATED_UCLASS_BODY()
+	GENERATED_BODY()
+public:
+	UEditorTutorialImportFactory(const FObjectInitializer& ObjectInitializer = FObjectInitializer::Get());
 
 	// UFactory Interface
 	virtual bool FactoryCanImport(const FString& Filename) override;
diff --git a/Engine/Source/Editor/IntroTutorials/Private/EditorTutorialSettings.h b/Engine/Source/Editor/IntroTutorials/Private/EditorTutorialSettings.h
--- a/Engine/Source/Editor/IntroTutorials/Private/EditorTutorialSettings.h
+++ b/Engine/Source/Editor/IntroTutorials/Private/EditorTutorialSettings.h
@@ -9,7 +9,7 @@
 USTRUCT()
 struct FTutorialContext
 {
-	GENERATED_USTRUCT_BODY()
+	GENERATED_BODY()
 
 	/** The context that this tutorial is used in */
 	UPROPERTY(EditAnywhere, Category = "Tutorials")
@@ -32,7 +32,9 @@ struct FTutorialContext
 UCLASS(config=EditorGameAgnostic)
 class UEditorTutorialSettings : public UObject
 {
-	GENERATED_UCLASS_BODY()
+	GENERATED_BODY()
+public:
+	UEditorTutorialSettings(const FObjectInitializer& ObjectInitializer = FObjectInitializer::Get());
 
 	/** Categories for tutorials */
 	UPROPERTY(Config, EditAnywhere, Category="Tutorials")
@@ -48,4 +50,4 @@ class UEditorTutorialSettings : public UObject
 
 	/** Get the tutorial info for the specified context */
 	void FindTutorialInfoForContext(FName InContext, UEditorTutorial*& OutAttractTutorial, UEditorTutorial*& OutLaunchTutorial, FString& OutBrowserFilter) const;
-};
\ No newline at end of file
+};
diff --git a/Engine/Source/Editor/IntroTutorials/Private/TutorialSettings.h b/Engine/Source/Editor/IntroTutorials/Private/TutorialSettings.h
--- a/Engine/Source/Editor/IntroTutorials/Private/TutorialSettings.h
+++ b/Engine/Source/Editor/IntroTutorials/Private/TutorialSettings.h
@@ -9,7 +9,9 @@
 UCLASS(config=Editor)
 class UTutorialSettings : public UObject
 {
-	GENERATED_UCLASS_BODY()
+	GENERATED_BODY()
+public:
+	UTutorialSettings(const FObjectInitializer& ObjectInitializer = FObjectInitializer::Get());
 
 	/** Categories for tutorials */
 	UPROPERTY(Config, EditAnywhere, Category="Tutorials")
@@ -18,4 +20,4 @@ class UTutorialSettings : public UObject
 	/** Tutorial to start on project startup */
 	UPROPERTY(Config, EditAnywhere, Category="Tutorials", meta=(MetaClass="EditorTutorial"))
 	FStringClassReference StartupTutorial;
-};
\ No newline at end of file
+};
diff --git a/Engine/Source/Editor/IntroTutorials/Private/TutorialStateSettings.h b/Engine/Source/Editor/IntroTutorials/Private/TutorialStateSettings.h
--- a/Engine/Source/Editor/IntroTutorials/Private/TutorialStateSettings.h
+++ b/Engine/Source/Editor/IntroTutorials/Private/TutorialStateSettings.h
@@ -9,7 +9,7 @@
 USTRUCT()
 struct FTutorialProgress
 {
-	GENERATED_USTRUCT_BODY()
+	GENERATED_BODY()
 
 	FTutorialProgress()
 	{
@@ -33,7 +33,9 @@ struct FTutorialProgress
 UCLASS(config=EditorGameAgnostic)
 class UTutorialStateSettings : public UObject
 {
-	GENERATED_UCLASS_BODY()
+	GENERATED_BODY()
+public:
+	UTutorialStateSettings(const FObjectInitializer& ObjectInitializer = FObjectInitializer::Get());
 
 	UPROPERTY(Config)
 	TArray<FTutorialProgress> TutorialsProgress;
@@ -68,4 +70,4 @@ class UTutorialStateSettings : public UObject
 private:
 	/** Recorded progress */
 	TMap<UEditorTutorial*, FTutorialProgress> ProgressMap;
-};
\ No newline at end of file
+};
diff --git a/Engine/Source/Editor/IntroTutorials/Public/EditorTutorial.h b/Engine/Source/Editor/IntroTutorials/Public/EditorTutorial.h
--- a/Engine/Source/Editor/IntroTutorials/Public/EditorTutorial.h
+++ b/Engine/Source/Editor/IntroTutorials/Public/EditorTutorial.h
@@ -45,7 +45,7 @@ namespace ETutorialAnchorIdentifier
 USTRUCT()
 struct INTROTUTORIALS_API FTutorialCategory
 {
-	GENERATED_USTRUCT_BODY()
+	GENERATED_BODY()
 
 	/** Period-separated category name, e.g. "Editor Quickstart.Level Editor" */
 	UPROPERTY(EditAnywhere, Category="Content")
@@ -76,7 +76,7 @@ struct INTROTUTORIALS_API FTutorialCategory
 USTRUCT()
 struct INTROTUTORIALS_API FTutorialContent
 {
-	GENERATED_USTRUCT_BODY()
+	GENERATED_BODY()
 
 	FTutorialContent()
 	{
@@ -104,7 +104,7 @@ struct INTROTUTORIALS_API FTutorialContent
 USTRUCT()
 struct INTROTUTORIALS_API FTutorialContentAnchor
 {
-	GENERATED_USTRUCT_BODY()
+	GENERATED_BODY()
 	
 	FTutorialContentAnchor()
 	{
@@ -149,7 +149,7 @@ struct INTROTUTORIALS_API FTutorialContentAnchor
 USTRUCT()
 struct INTROTUTORIALS_API FTutorialWidgetContent
 {
-	GENERATED_USTRUCT_BODY()
+	GENERATED_BODY()
 
 	FTutorialWidgetContent()
 	{
@@ -190,7 +190,7 @@ struct INTROTUTORIALS_API FTutorialWidgetContent
 USTRUCT()
 struct INTROTUTORIALS_API FTutorialStage
 {
-	GENERATED_USTRUCT_BODY()
+	GENERATED_BODY()
 
 	/** Identifier for this stage */
 	UPROPERTY(EditAnywhere, Category="Stage")
@@ -217,7 +217,9 @@ struct INTROTUTORIALS_API FTutorialStage
 UCLASS(Blueprintable, hideCategories=(Object), EditInlineNew)
 class INTROTUTORIALS_API UEditorTutorial : public UObject
 {
-	GENERATED_UCLASS_BODY()
+	GENERATED_BODY()
+public:
+	UEditorTutorial(const FObjectInitializer& ObjectInitializer = FObjectInitializer::Get());
 
 	/** Title of this tutorial, used when presented to the user */
 	UPROPERTY(EditAnywhere, Category="Tutorial", AssetRegistrySearchable)
@@ -345,4 +347,4 @@ protected:
 	 */
 	UFUNCTION(BlueprintCallable, Category="Tutorial")
 	static bool GetEngineFolderVisibilty();
-};
\ No newline at end of file
+};
diff --git a/Engine/Source/Editor/Kismet/Classes/BlueprintPaletteFavorites.h b/Engine/Source/Editor/Kismet/Classes/BlueprintPaletteFavorites.h
--- a/Engine/Source/Editor/Kismet/Classes/BlueprintPaletteFavorites.h
+++ b/Engine/Source/Editor/Kismet/Classes/BlueprintPaletteFavorites.h
@@ -16,7 +16,7 @@ struct FBlueprintActionInfo;
 USTRUCT()
 struct FFavoritedBlueprintPaletteItem
 {
-	GENERATED_USTRUCT_BODY()
+	GENERATED_BODY()
 
 public:
 	/**
@@ -88,7 +88,9 @@ private:
 UCLASS(config=EditorUserSettings)
 class KISMET_API UBlueprintPaletteFavorites : public UObject
 {
-	GENERATED_UCLASS_BODY()
+	GENERATED_BODY()
+public:
+	UBlueprintPaletteFavorites(const FObjectInitializer& ObjectInitializer = FObjectInitializer::Get());
 
 public:
 	// Begin UObject interface
diff --git a/Engine/Source/Editor/Kismet/Public/BlueprintActionMenuItem.h b/Engine/Source/Editor/Kismet/Public/BlueprintActionMenuItem.h
--- a/Engine/Source/Editor/Kismet/Public/BlueprintActionMenuItem.h
+++ b/Engine/Source/Editor/Kismet/Public/BlueprintActionMenuItem.h
@@ -22,7 +22,7 @@ struct FBlueprintActionContext;
 USTRUCT()
 struct FBlueprintActionMenuItem : public FEdGraphSchemaAction
 {
-	GENERATED_USTRUCT_BODY();
+	GENERATED_BODY()
 	static FName StaticGetTypeId() { static FName const TypeId("FBlueprintActionMenuItem"); return TypeId; }
 
 public:	
diff --git a/Engine/Source/Editor/Kismet/Public/BlueprintDragDropMenuItem.h b/Engine/Source/Editor/Kismet/Public/BlueprintDragDropMenuItem.h
--- a/Engine/Source/Editor/Kismet/Public/BlueprintDragDropMenuItem.h
+++ b/Engine/Source/Editor/Kismet/Public/BlueprintDragDropMenuItem.h
@@ -29,7 +29,7 @@ struct FBlueprintActionContext;
 USTRUCT()
 struct FBlueprintDragDropMenuItem : public FEdGraphSchemaAction
 {
-	GENERATED_USTRUCT_BODY();
+	GENERATED_BODY()
 	
 public:
 	static FName StaticGetTypeId() { static FName const TypeId("FBlueprintDragDropMenuItem"); return TypeId; }
diff --git a/Engine/Source/Editor/Kismet/Public/Nodes/K2Node_AsyncAction.h b/Engine/Source/Editor/Kismet/Public/Nodes/K2Node_AsyncAction.h
--- a/Engine/Source/Editor/Kismet/Public/Nodes/K2Node_AsyncAction.h
+++ b/Engine/Source/Editor/Kismet/Public/Nodes/K2Node_AsyncAction.h
@@ -7,7 +7,9 @@
 UCLASS()
 class KISMET_API UK2Node_AsyncAction : public UK2Node_BaseAsyncTask
 {
-	GENERATED_UCLASS_BODY()
+	GENERATED_BODY()
+public:
+	UK2Node_AsyncAction(const FObjectInitializer& ObjectInitializer = FObjectInitializer::Get());
 	
 	// UEdGraphNode interface
 	virtual void GetMenuEntries(FGraphContextMenuBuilder& ContextMenuBuilder) const override;
diff --git a/Engine/Source/Editor/KismetCompiler/Private/KismetCompilerCppBackend.cpp b/Engine/Source/Editor/KismetCompiler/Private/KismetCompilerCppBackend.cpp
--- a/Engine/Source/Editor/KismetCompiler/Private/KismetCompilerCppBackend.cpp
+++ b/Engine/Source/Editor/KismetCompiler/Private/KismetCompilerCppBackend.cpp
@@ -774,7 +774,8 @@ void FKismetCppBackend::GenerateCodeFromClass(UClass* SourceClass, TIndirectArra
 	Emit(Header,
 		TEXT("{\n")
 		TEXT("public:\n"));
-	Emit(Header, *FString::Printf(TEXT("\tGENERATED_UCLASS_BODY()\n")));
+	Emit(Header, *FString::Printf(TEXT("\tGENERATED_BODY()\n")));
+	Emit(Header, *FString::Printf(TEXT("public:\n")));
 
 	EmitClassProperties(Header, SourceClass);
 
@@ -794,11 +795,6 @@ void FKismetCppBackend::GenerateCodeFromClass(UClass* SourceClass, TIndirectArra
 	Emit(Body, *FString::Printf(TEXT("#include \"%s.h\"\n"), FApp::GetGameName()));
 	Emit(Body, TEXT("#include \"GeneratedCodeHelpers.h\"\n\n"));
 
-	//constructor
-	Emit(Body, *FString::Printf(
-		TEXT("%s::%s(const FObjectInitializer& ObjectInitializer) : Super(ObjectInitializer) {}\n\n"),
-		*CppClassName, *CppClassName));
-
 	for (int32 i = 0; i < Functions.Num(); ++i)
 	{
 		if (Functions[i].IsValid())
diff --git a/Engine/Source/Editor/LandscapeEditor/Classes/ActorFactoryLandscape.h b/Engine/Source/Editor/LandscapeEditor/Classes/ActorFactoryLandscape.h
--- a/Engine/Source/Editor/LandscapeEditor/Classes/ActorFactoryLandscape.h
+++ b/Engine/Source/Editor/LandscapeEditor/Classes/ActorFactoryLandscape.h
@@ -7,7 +7,9 @@
 UCLASS(MinimalAPI, config=Editor)
 class UActorFactoryLandscape : public UActorFactory
 {
-	GENERATED_UCLASS_BODY()
+	GENERATED_BODY()
+public:
+	LANDSCAPEEDITOR_API UActorFactoryLandscape(const FObjectInitializer& ObjectInitializer = FObjectInitializer::Get());
 
 	// Begin UActorFactory Interface
 	//virtual bool CanCreateActorFrom( const FAssetData& AssetData, FText& OutErrorMsg ) override;
diff --git a/Engine/Source/Editor/LandscapeEditor/Classes/LandscapeEditorObject.h b/Engine/Source/Editor/LandscapeEditor/Classes/LandscapeEditorObject.h
--- a/Engine/Source/Editor/LandscapeEditor/Classes/LandscapeEditorObject.h
+++ b/Engine/Source/Editor/LandscapeEditor/Classes/LandscapeEditorObject.h
@@ -128,7 +128,7 @@ namespace EColorChannel
 USTRUCT()
 struct FGizmoImportLayer
 {
-	GENERATED_USTRUCT_BODY()
+	GENERATED_BODY()
 
 	UPROPERTY(Category="Import", EditAnywhere)
 	FString LayerFilename;
@@ -175,7 +175,7 @@ enum class ELandscapeImportLayerError : uint8
 USTRUCT()
 struct FLandscapeImportLayer : public FLandscapeImportLayerInfo
 {
-	GENERATED_USTRUCT_BODY()
+	GENERATED_BODY()
 
 	UPROPERTY(Category="Import", VisibleAnywhere)
 	ELandscapeImportLayerError ImportError;
@@ -190,7 +190,7 @@ struct FLandscapeImportLayer : public FLandscapeImportLayerInfo
 USTRUCT()
 struct FLandscapePatternBrushWorldSpaceSettings
 {
-	GENERATED_USTRUCT_BODY()
+	GENERATED_BODY()
 
 	UPROPERTY(Category="World-Space", EditAnywhere)
 	FVector2D Origin;
@@ -212,7 +212,9 @@ struct FLandscapePatternBrushWorldSpaceSettings
 UCLASS()
 class ULandscapeEditorObject : public UObject
 {
-	GENERATED_UCLASS_BODY()
+	GENERATED_BODY()
+public:
+	ULandscapeEditorObject(const FObjectInitializer& ObjectInitializer = FObjectInitializer::Get());
 
 	FEdModeLandscape* ParentMode;
 
diff --git a/Engine/Source/Editor/LandscapeEditor/Classes/LandscapePlaceholder.h b/Engine/Source/Editor/LandscapeEditor/Classes/LandscapePlaceholder.h
--- a/Engine/Source/Editor/LandscapeEditor/Classes/LandscapePlaceholder.h
+++ b/Engine/Source/Editor/LandscapeEditor/Classes/LandscapePlaceholder.h
@@ -7,7 +7,9 @@
 UCLASS(NotPlaceable)
 class ALandscapePlaceholder : public AActor
 {
-	GENERATED_UCLASS_BODY()
+	GENERATED_BODY()
+public:
+	ALandscapePlaceholder(const FObjectInitializer& ObjectInitializer = FObjectInitializer::Get());
 
 	//virtual void PostEditMove(bool bFinished) override;
 	virtual bool TeleportTo(const FVector& DestLocation, const FRotator& DestRotation, bool bIsATest = false, bool bNoCheck = false) override;
diff --git a/Engine/Source/Editor/Levels/Public/LevelBrowserSettings.h b/Engine/Source/Editor/Levels/Public/LevelBrowserSettings.h
--- a/Engine/Source/Editor/Levels/Public/LevelBrowserSettings.h
+++ b/Engine/Source/Editor/Levels/Public/LevelBrowserSettings.h
@@ -16,7 +16,9 @@ UCLASS(config=EditorUserSettings)
 class LEVELS_API ULevelBrowserSettings
 	:	public UObject
 {
-	GENERATED_UCLASS_BODY()
+	GENERATED_BODY()
+public:
+	ULevelBrowserSettings(const FObjectInitializer& ObjectInitializer = FObjectInitializer::Get());
 
 public:
 
diff --git a/Engine/Source/Editor/LocalizationDashboard/Classes/LocalizationTarget.h b/Engine/Source/Editor/LocalizationDashboard/Classes/LocalizationTarget.h
--- a/Engine/Source/Editor/LocalizationDashboard/Classes/LocalizationTarget.h
+++ b/Engine/Source/Editor/LocalizationDashboard/Classes/LocalizationTarget.h
@@ -16,7 +16,7 @@ enum class ELocalizationTargetStatus : uint8
 USTRUCT()
 struct FGatherTextFromTextFilesConfiguration
 {
-	GENERATED_USTRUCT_BODY()
+	GENERATED_BODY()
 
 	UPROPERTY(config, EditAnywhere, Category = "Filter")
 	TArray<FString> SearchDirectories;
@@ -31,7 +31,7 @@ struct FGatherTextFromTextFilesConfiguration
 USTRUCT()
 struct FGatherTextFromPackagesConfiguration
 {
-	GENERATED_USTRUCT_BODY()
+	GENERATED_BODY()
 
 	UPROPERTY(config, EditAnywhere, Category = "Filter")
 	TArray<FString> IncludePathWildcards;
@@ -46,7 +46,7 @@ struct FGatherTextFromPackagesConfiguration
 USTRUCT()
 struct FMetaDataKeyGatherSpecification
 {
-	GENERATED_USTRUCT_BODY()
+	GENERATED_BODY()
 
 	UPROPERTY(config, EditAnywhere, Category = "Input")
 	FString MetaDataKey;
@@ -61,7 +61,7 @@ struct FMetaDataKeyGatherSpecification
 USTRUCT()
 struct FGatherTextFromMetaDataConfiguration
 {
-	GENERATED_USTRUCT_BODY()
+	GENERATED_BODY()
 
 	UPROPERTY(config, EditAnywhere, Category = "Filter")
 	TArray<FString> IncludePathWildcards;
@@ -76,7 +76,7 @@ struct FGatherTextFromMetaDataConfiguration
 USTRUCT()
 struct FCultureStatistics
 {
-	GENERATED_USTRUCT_BODY()
+	GENERATED_BODY()
 
 	FCultureStatistics()
 		: WordCount(0)
@@ -93,7 +93,7 @@ struct FCultureStatistics
 USTRUCT()
 struct FLocalizationTargetSettings
 {
-	GENERATED_USTRUCT_BODY()
+	GENERATED_BODY()
 
 	FLocalizationTargetSettings();
 
diff --git a/Engine/Source/Editor/Matinee/Classes/InterpTrackHelper.h b/Engine/Source/Editor/Matinee/Classes/InterpTrackHelper.h
--- a/Engine/Source/Editor/Matinee/Classes/InterpTrackHelper.h
+++ b/Engine/Source/Editor/Matinee/Classes/InterpTrackHelper.h
@@ -8,7 +8,9 @@
 UCLASS()
 class MATINEE_API UInterpTrackHelper : public UObject
 {
-	GENERATED_UCLASS_BODY()
+	GENERATED_BODY()
+public:
+	UInterpTrackHelper(const FObjectInitializer& ObjectInitializer = FObjectInitializer::Get());
 
 	/**
 	 * @param  Track	The track to get the actor for.
diff --git a/Engine/Source/Editor/Matinee/Classes/MatineeOptions.h b/Engine/Source/Editor/Matinee/Classes/MatineeOptions.h
--- a/Engine/Source/Editor/Matinee/Classes/MatineeOptions.h
+++ b/Engine/Source/Editor/Matinee/Classes/MatineeOptions.h
@@ -11,7 +11,9 @@
 UCLASS(hidecategories=Object, config=Editor)
 class UMatineeOptions : public UObject
 {
-	GENERATED_UCLASS_BODY()
+	GENERATED_BODY()
+public:
+	UMatineeOptions(const FObjectInitializer& ObjectInitializer = FObjectInitializer::Get());
 
 	UPROPERTY()
 	TArray<struct FInterpEdSelKey> SelectedKeys;
diff --git a/Engine/Source/Editor/Matinee/Classes/MatineeTrackAnimControlHelper.h b/Engine/Source/Editor/Matinee/Classes/MatineeTrackAnimControlHelper.h
--- a/Engine/Source/Editor/Matinee/Classes/MatineeTrackAnimControlHelper.h
+++ b/Engine/Source/Editor/Matinee/Classes/MatineeTrackAnimControlHelper.h
@@ -8,7 +8,9 @@
 UCLASS()
 class UMatineeTrackAnimControlHelper : public UInterpTrackHelper
 {
-	GENERATED_UCLASS_BODY()
+	GENERATED_BODY()
+public:
+	UMatineeTrackAnimControlHelper(const FObjectInitializer& ObjectInitializer = FObjectInitializer::Get());
 
 	void OnAddKeyTextEntry(const class FAssetData& AssetData, IMatineeBase* Matinee, UInterpTrack* Track);
 	void OnCreateTrackTextEntry(const FString& ChosenText, TSharedRef<SWindow> Window, FString* OutputString);
diff --git a/Engine/Source/Editor/Matinee/Classes/MatineeTrackBoolPropHelper.h b/Engine/Source/Editor/Matinee/Classes/MatineeTrackBoolPropHelper.h
--- a/Engine/Source/Editor/Matinee/Classes/MatineeTrackBoolPropHelper.h
+++ b/Engine/Source/Editor/Matinee/Classes/MatineeTrackBoolPropHelper.h
@@ -7,7 +7,9 @@
 UCLASS()
 class UMatineeTrackBoolPropHelper : public UInterpTrackHelper
 {
-	GENERATED_UCLASS_BODY()
+	GENERATED_BODY()
+public:
+	UMatineeTrackBoolPropHelper(const FObjectInitializer& ObjectInitializer = FObjectInitializer::Get());
 
 	void OnCreateTrackTextEntry(const FString& ChosenText, TWeakPtr<SWindow> Window, FString* OutputString);
 
