Dismissing a the tutorial on the Tutorials Browser button stops it from coming back when the editor is restarted, or from launching again when you try to open the browser.

diff --git a/Engine/Source/Editor/IntroTutorials/Private/STutorialButton.cpp b/Engine/Source/Editor/IntroTutorials/Private/STutorialButton.cpp
--- a/Engine/Source/Editor/IntroTutorials/Private/STutorialButton.cpp
+++ b/Engine/Source/Editor/IntroTutorials/Private/STutorialButton.cpp
@@ -302,7 +302,7 @@ void STutorialButton::LaunchBrowser()
 
 bool STutorialButton::ShouldLaunchBrowser() const
 {
-	return (!bTutorialAvailable || (bTutorialAvailable && bTutorialCompleted));
+	return (!bTutorialAvailable || bTutorialCompleted || bTutorialDismissed);
 }
 
 bool STutorialButton::ShouldShowAlert() const
