[UE-10992] Dismiss-all added to right-click menu on tutorial button. Stops green pulse animation on all tutorial buttons and saves to config file.
Known minor bug: If multiple windows are open with pulses animating, only the one that is given the "dismiss all" command will stop animating. The others will continue until their windows are closed and reopened, or they are dismissed individually.

diff --git a/Engine/Source/Editor/IntroTutorials/Private/STutorialButton.cpp b/Engine/Source/Editor/IntroTutorials/Private/STutorialButton.cpp
--- a/Engine/Source/Editor/IntroTutorials/Private/STutorialButton.cpp
+++ b/Engine/Source/Editor/IntroTutorials/Private/STutorialButton.cpp
@@ -61,12 +61,12 @@ EActiveTimerReturnType STutorialButton::OpenTutorialPostConstruct( double InCurr
 
 	RefreshStatus();
 
-	if ( bTutorialAvailable && CachedAttractTutorial != nullptr && !bTutorialDismissed && !bTutorialCompleted )
+	if (bTutorialAvailable && CachedAttractTutorial != nullptr && !bTutorialDismissed && !bTutorialCompleted && !GetMutableDefault<UTutorialStateSettings>()->AreAllTutorialsDismissed())
 	{
 		// kick off the attract tutorial if the user hasn't dismissed it and hasn't completed it
 		FIntroTutorials& IntroTutorials = FModuleManager::GetModuleChecked<FIntroTutorials>( TEXT( "IntroTutorials" ) );
 		const bool bRestart = true;
-		IntroTutorials.LaunchTutorial( CachedAttractTutorial, bRestart, ContextWindow );
+		IntroTutorials.LaunchTutorial(CachedAttractTutorial, bRestart, ContextWindow);
 	}
 
 	if ( ShouldShowAlert() )
@@ -230,6 +230,13 @@ FReply STutorialButton::OnMouseButtonDown(const FGeometry& MyGeometry, const FPo
 				FUIAction(FExecuteAction::CreateSP(this, &STutorialButton::DismissAlert))
 				);
 
+			MenuBuilder.AddMenuEntry(
+				LOCTEXT("DismissAllReminders", "Dismiss All Tutorial Reminders"),
+				LOCTEXT("DismissAllRemindersTooltip", "Selecting this option will prevent all tutorial blips from being displayed."),
+				FSlateIcon(),
+				FUIAction(FExecuteAction::CreateSP(this, &STutorialButton::DismissAllAlerts))
+				);
+
 			MenuBuilder.AddMenuSeparator();
 		}
 
@@ -287,6 +294,12 @@ void STutorialButton::DismissAlert()
 	RefreshStatus();
 }
 
+void STutorialButton::DismissAllAlerts()
+{
+	GetMutableDefault<UTutorialStateSettings>()->DismissAllTutorials();
+	DismissAlert();		//TODO FIXME Call this for all STutorialButtons that are currently displayed so they all stop pulsing.
+}
+
 void STutorialButton::LaunchTutorial()
 {
 	HandleButtonClicked();
@@ -307,7 +320,11 @@ bool STutorialButton::ShouldLaunchBrowser() const
 
 bool STutorialButton::ShouldShowAlert() const
 {
-	return ((bTestAlerts || !FEngineBuildSettings::IsInternalBuild()) && bTutorialAvailable && !(bTutorialCompleted || bTutorialDismissed));
+	if ((bTestAlerts || !FEngineBuildSettings::IsInternalBuild()) && bTutorialAvailable && !(bTutorialCompleted || bTutorialDismissed))
+	{
+		return !GetMutableDefault<UTutorialStateSettings>()->AreAllTutorialsDismissed();
+	}
+	return false;
 }
 
 FText STutorialButton::GetButtonToolTip() const
diff --git a/Engine/Source/Editor/IntroTutorials/Private/STutorialButton.h b/Engine/Source/Editor/IntroTutorials/Private/STutorialButton.h
--- a/Engine/Source/Editor/IntroTutorials/Private/STutorialButton.h
+++ b/Engine/Source/Editor/IntroTutorials/Private/STutorialButton.h
@@ -33,6 +33,9 @@ private:
 	/** Dismiss the pulsing alert */
 	void DismissAlert();
 
+	/** Dismiss all pulsing alerts */
+	void DismissAllAlerts();
+
 	/** Launch the tutorials browser */
 	void LaunchBrowser();
 
diff --git a/Engine/Source/Editor/IntroTutorials/Private/TutorialStateSettings.cpp b/Engine/Source/Editor/IntroTutorials/Private/TutorialStateSettings.cpp
--- a/Engine/Source/Editor/IntroTutorials/Private/TutorialStateSettings.cpp
+++ b/Engine/Source/Editor/IntroTutorials/Private/TutorialStateSettings.cpp
@@ -6,7 +6,7 @@
 UTutorialStateSettings::UTutorialStateSettings(const FObjectInitializer& ObjectInitializer)
 	: Super(ObjectInitializer)
 {
-
+	bDismissedAllTutorials = false;
 }
 
 void UTutorialStateSettings::PostInitProperties()
@@ -99,6 +99,11 @@ void UTutorialStateSettings::DismissTutorial(UEditorTutorial* InTutorial, bool b
 
 bool UTutorialStateSettings::IsTutorialDismissed(UEditorTutorial* InTutorial) const
 {
+	if (GetMutableDefault<UTutorialStateSettings>()->AreAllTutorialsDismissed())
+	{
+		return true;
+	}
+
 	const FTutorialProgress* FoundProgress = ProgressMap.Find(InTutorial);
 	if (FoundProgress != nullptr)
 	{
@@ -107,6 +112,16 @@ bool UTutorialStateSettings::IsTutorialDismissed(UEditorTutorial* InTutorial) co
 	return false;
 }
 
+void UTutorialStateSettings::DismissAllTutorials()
+{
+	bDismissedAllTutorials = true;
+}
+
+bool UTutorialStateSettings::AreAllTutorialsDismissed()
+{
+	return bDismissedAllTutorials;
+}
+
 void UTutorialStateSettings::SaveProgress()
 {
 	TutorialsProgress.Empty();
@@ -123,6 +138,7 @@ void UTutorialStateSettings::ClearProgress()
 {
 	ProgressMap.Empty();
 	TutorialsProgress.Empty();
+	bDismissedAllTutorials = false;
 
 	SaveConfig();
 }
\ No newline at end of file
diff --git a/Engine/Source/Editor/IntroTutorials/Private/TutorialStateSettings.h b/Engine/Source/Editor/IntroTutorials/Private/TutorialStateSettings.h
--- a/Engine/Source/Editor/IntroTutorials/Private/TutorialStateSettings.h
+++ b/Engine/Source/Editor/IntroTutorials/Private/TutorialStateSettings.h
@@ -65,7 +65,17 @@ class UTutorialStateSettings : public UObject
 	/** Save the progress of all our tutorials */
 	void SaveProgress();
 
+	/** Dismiss all tutorials, used by right-click option on scholar cap button (STutorialButton) */
+	void DismissAllTutorials();
+
+	/** Returns true if user has dismissed tutorials */
+	bool AreAllTutorialsDismissed();
+
 private:
 	/** Recorded progress */
 	TMap<UEditorTutorial*, FTutorialProgress> ProgressMap;
+
+	/** Record if user has chosen to cancel all tutorials */
+	UPROPERTY(Config)
+	bool bDismissedAllTutorials;
 };
\ No newline at end of file
