[UE-9386] Tutorial Browser categories show up as completed if all tutorials in the category are completed. Also works recursively if there's a category within a category, meaning the sub-category must have all its tutorials (and sub-categories) complete in order to show as complete itself.

diff --git a/Engine/Source/Editor/IntroTutorials/Private/STutorialsBrowser.cpp b/Engine/Source/Editor/IntroTutorials/Private/STutorialsBrowser.cpp
--- a/Engine/Source/Editor/IntroTutorials/Private/STutorialsBrowser.cpp
+++ b/Engine/Source/Editor/IntroTutorials/Private/STutorialsBrowser.cpp
@@ -101,17 +101,30 @@ public:
 					.HAlign(HAlign_Center)
 					.Padding(8.0f)
 					[
-						SNew(SBox)
-						.WidthOverride(64.0f)
-						.HeightOverride(64.0f)
-						.VAlign(VAlign_Center)
-						.HAlign(HAlign_Center)
+						SNew(SOverlay)
+						+ SOverlay::Slot()
+						[
+							SNew(SBox)
+							.WidthOverride(64.0f)
+							.HeightOverride(64.0f)
+							.VAlign(VAlign_Center)
+							.HAlign(HAlign_Center)
+							[
+								SNew(SImage)
+								.Image(SlateBrush)
+							]
+						]
+						+ SOverlay::Slot()
+						.VAlign(VAlign_Bottom)
+						.HAlign(HAlign_Right)
 						[
 							SNew(SImage)
-							.Image(SlateBrush)
+							.ToolTipText(LOCTEXT("CompletedCheckToolTip", "This category has been completed"))
+							.Visibility(this, &FTutorialListEntry_Category::GetCompletedVisibility)
+							.Image(FEditorStyle::Get().GetBrush("Tutorials.Browser.Completed"))
 						]
 					]
-					+SHorizontalBox::Slot()
+					+ SHorizontalBox::Slot()
 					.FillWidth(1.0f)
 					.VAlign(VAlign_Center)
 					[
@@ -205,6 +218,25 @@ public:
 		return (SubCategories.Num() > 0 || Tutorials.Num() > 0) ? EVisibility::Visible : EVisibility::Collapsed;
 	}
 
+	EVisibility GetCompletedVisibility() const override
+	{
+		for (int32 i = 0; i < Tutorials.Num(); ++i)
+		{
+			if (Tutorials[i].IsValid() && (Tutorials[i]->GetCompletedVisibility() != EVisibility::Visible))
+			{
+				return EVisibility::Hidden;
+			}
+		}
+		for (int32 i = 0; i < SubCategories.Num(); ++i)
+		{
+			if (SubCategories[i].IsValid() && (SubCategories[i]->GetCompletedVisibility() != EVisibility::Visible))
+			{
+				return EVisibility::Hidden;
+			}
+		}
+		return EVisibility::Visible;
+	}
+
 public:
 	/** Copy of the category info */
 	FTutorialCategory Category;
@@ -437,7 +469,7 @@ public:
 		return EVisibility::Collapsed;
 	}
 
-	EVisibility GetCompletedVisibility() const
+	EVisibility GetCompletedVisibility() const override
 	{
 		CacheProgress();
 		return bHaveCompletedTutorial ? EVisibility::Visible : EVisibility::Hidden;
diff --git a/Engine/Source/Editor/IntroTutorials/Private/STutorialsBrowser.h b/Engine/Source/Editor/IntroTutorials/Private/STutorialsBrowser.h
--- a/Engine/Source/Editor/IntroTutorials/Private/STutorialsBrowser.h
+++ b/Engine/Source/Editor/IntroTutorials/Private/STutorialsBrowser.h
@@ -29,6 +29,9 @@ struct ITutorialListEntry
 	 * @return true if this < OtherEntry 
 	 */
 	virtual bool SortAgainst(TSharedRef<ITutorialListEntry> OtherEntry) const = 0;
+
+	/** Return true if this entry should show up as completed (currently used to hide/show green check mark) */
+	virtual EVisibility GetCompletedVisibility() const = 0;
 };
 
 class FTutorialListEntry_Category;
