[UE-7933] Improved on previous fix so the tutorial itself closes the browser on launch, rather than the browser doing it.

diff --git a/Engine/Source/Editor/IntroTutorials/Private/EditorTutorial.cpp b/Engine/Source/Editor/IntroTutorials/Private/EditorTutorial.cpp
--- a/Engine/Source/Editor/IntroTutorials/Private/EditorTutorial.cpp
+++ b/Engine/Source/Editor/IntroTutorials/Private/EditorTutorial.cpp
@@ -1,6 +1,8 @@
 // Copyright 1998-2015 Epic Games, Inc. All Rights Reserved.
 
 #include "IntroTutorialsPrivatePCH.h"
+#include "LevelEditor.h"
+#include "SDockTab.h"
 #include "EditorTutorial.h"
 
 UEditorTutorial::UEditorTutorial(const FObjectInitializer& ObjectInitializer)
@@ -57,6 +59,8 @@ void UEditorTutorial::HandleTutorialStageEnded(FName StageName)
 void UEditorTutorial::HandleTutorialLaunched()
 {
 	FEditorScriptExecutionGuard ScriptGuard;
+	FLevelEditorModule& LevelEditorModule = FModuleManager::GetModuleChecked<FLevelEditorModule>(TEXT("LevelEditor"));
+	LevelEditorModule.GetLevelEditorTabManager()->InvokeTab(FTabId("TutorialsBrowser"))->RequestCloseTab();
 	OnTutorialLaunched();
 }
 
diff --git a/Engine/Source/Editor/IntroTutorials/Private/STutorialButton.cpp b/Engine/Source/Editor/IntroTutorials/Private/STutorialButton.cpp
--- a/Engine/Source/Editor/IntroTutorials/Private/STutorialButton.cpp
+++ b/Engine/Source/Editor/IntroTutorials/Private/STutorialButton.cpp
@@ -8,8 +8,6 @@
 #include "TutorialMetaData.h"
 #include "EngineBuildSettings.h"
 #include "AssetRegistryModule.h"
-#include "LevelEditor.h"
-#include "SDockTab.h"
 
 #define LOCTEXT_NAMESPACE "STutorialButton"
 
@@ -182,8 +180,6 @@ EActiveTimerReturnType STutorialButton::HandleButtonClicked_AssetRegistryChecker
 	//Sometimes, this gives a false positive because the tutorial we want to launch wasn't loaded into the asset registry when we checked. Opening and closing the tab works around that by letting the browser recheck.
 	if (ShouldLaunchBrowser())
 	{
-		FLevelEditorModule& LevelEditorModule = FModuleManager::GetModuleChecked<FLevelEditorModule>(TEXT("LevelEditor"));
-		LevelEditorModule.GetLevelEditorTabManager()->InvokeTab(FTabId("TutorialsBrowser"))->RequestCloseTab();
 		RefreshStatus();
 	}
 
