[UE-9586] The "Back" button in tutorials can be linked to a different tutorial, like the "Next" button. Tutorials assigned to the "Previous Tutorial" field (new) will be used for this, and will activate the button for this purpose on stage 0. Tutorials started in this way will begin on their final stage. The Launch Tutorial functionality (in code) now uses an enum for startup instead of a bool, so we can force a restart, continue from where we left off (default), or start at the final stage.

diff --git a/Engine/Source/Developer/SuperSearch/Private/SSuperSearch.cpp b/Engine/Source/Developer/SuperSearch/Private/SSuperSearch.cpp
--- a/Engine/Source/Developer/SuperSearch/Private/SSuperSearch.cpp
+++ b/Engine/Source/Developer/SuperSearch/Private/SSuperSearch.cpp
@@ -170,7 +170,7 @@ void SSuperSearchBox::ActOnSuggestion(TSharedPtr<FSearchEntry> SearchEntry, FStr
 			FWidgetPath OutWidgetPath;
 			TSharedPtr<SWindow> ParentWindow = FSlateApplication::Get().FindWidgetWindow(SuggestionBox.ToSharedRef(), OutWidgetPath);
 
-			IntroTutorials.LaunchTutorial(Tutorial, true, ParentWindow);
+			IntroTutorials.LaunchTutorial(Tutorial, IIntroTutorials::ETutorialStartType::TST_RESTART, ParentWindow);
 		}
 		else
 #endif
diff --git a/Engine/Source/Editor/IntroTutorials/Private/EditorTutorial.cpp b/Engine/Source/Editor/IntroTutorials/Private/EditorTutorial.cpp
--- a/Engine/Source/Editor/IntroTutorials/Private/EditorTutorial.cpp
+++ b/Engine/Source/Editor/IntroTutorials/Private/EditorTutorial.cpp
@@ -38,7 +38,7 @@ void UEditorTutorial::GoToPreviousTutorialStage()
 void UEditorTutorial::BeginTutorial(UEditorTutorial* TutorialToStart, bool bRestart)
 {
 	FIntroTutorials& IntroTutorials = FModuleManager::GetModuleChecked<FIntroTutorials>(TEXT("IntroTutorials"));
-	IntroTutorials.LaunchTutorial(TutorialToStart, bRestart);
+	IntroTutorials.LaunchTutorial(TutorialToStart, bRestart ? IIntroTutorials::ETutorialStartType::TST_RESTART : IIntroTutorials::ETutorialStartType::TST_CONTINUE);
 }
 
 
diff --git a/Engine/Source/Editor/IntroTutorials/Private/EditorTutorialDetailsCustomization.cpp b/Engine/Source/Editor/IntroTutorials/Private/EditorTutorialDetailsCustomization.cpp
--- a/Engine/Source/Editor/IntroTutorials/Private/EditorTutorialDetailsCustomization.cpp
+++ b/Engine/Source/Editor/IntroTutorials/Private/EditorTutorialDetailsCustomization.cpp
@@ -21,7 +21,7 @@ void FEditorTutorialDetailsCustomization::CustomizeDetails( IDetailLayoutBuilder
 		{
 			FLevelEditorModule& LevelEditorModule = FModuleManager::Get().GetModuleChecked<FLevelEditorModule>("LevelEditor");
 			IIntroTutorials& IntroTutorials = FModuleManager::Get().GetModuleChecked<IIntroTutorials>("IntroTutorials");
-			IntroTutorials.LaunchTutorial(Tutorial, true, LevelEditorModule.GetLevelEditorTab()->GetParentWindow());
+			IntroTutorials.LaunchTutorial(Tutorial, IIntroTutorials::ETutorialStartType::TST_RESTART, LevelEditorModule.GetLevelEditorTab()->GetParentWindow());
 			return FReply::Handled();
 		}
 	};
diff --git a/Engine/Source/Editor/IntroTutorials/Private/IntroTutorials.cpp b/Engine/Source/Editor/IntroTutorials/Private/IntroTutorials.cpp
--- a/Engine/Source/Editor/IntroTutorials/Private/IntroTutorials.cpp
+++ b/Engine/Source/Editor/IntroTutorials/Private/IntroTutorials.cpp
@@ -306,8 +306,7 @@ void FIntroTutorials::SummonBlueprintTutorialHome(UObject* Asset, bool bForceWel
 			check(ContextWindow.IsValid());
 		}
 
-		const bool bRestart = true;
-		IntroTutorials.LaunchTutorial(LaunchTutorial, bRestart, ContextWindow);
+		IntroTutorials.LaunchTutorial(LaunchTutorial, IIntroTutorials::ETutorialStartType::TST_RESTART, ContextWindow);
 	}
 }
 
@@ -443,11 +442,11 @@ void FIntroTutorials::LaunchTutorialByName(const FString& InAssetPath, bool bInR
 	UBlueprint* Blueprint = LoadObject<UBlueprint>(nullptr, *InAssetPath);
 	if (Blueprint && Blueprint->GeneratedClass)
 	{
-		LaunchTutorial(Blueprint->GeneratedClass->GetDefaultObject<UEditorTutorial>(), bInRestart, InNavigationWindow, OnTutorialClosed, OnTutorialExited);
+		LaunchTutorial(Blueprint->GeneratedClass->GetDefaultObject<UEditorTutorial>(), bInRestart ? IIntroTutorials::ETutorialStartType::TST_RESTART : IIntroTutorials::ETutorialStartType::TST_CONTINUE, InNavigationWindow, OnTutorialClosed, OnTutorialExited);
 	}
 }
 
-void FIntroTutorials::LaunchTutorial(UEditorTutorial* InTutorial, bool bInRestart, TWeakPtr<SWindow> InNavigationWindow, FSimpleDelegate OnTutorialClosed, FSimpleDelegate OnTutorialExited)
+void FIntroTutorials::LaunchTutorial(UEditorTutorial* InTutorial, ETutorialStartType InStartType, TWeakPtr<SWindow> InNavigationWindow, FSimpleDelegate OnTutorialClosed, FSimpleDelegate OnTutorialExited)
 {
 	if(TutorialRoot.IsValid())
 	{
@@ -456,8 +455,7 @@ void FIntroTutorials::LaunchTutorial(UEditorTutorial* InTutorial, bool bInRestar
 			IMainFrameModule& MainFrameModule = FModuleManager::LoadModuleChecked<IMainFrameModule>(TEXT("MainFrame"));
 			InNavigationWindow = MainFrameModule.GetParentWindow();
 		}
-
-		TutorialRoot->LaunchTutorial(InTutorial, bInRestart, InNavigationWindow, OnTutorialClosed, OnTutorialExited);
+		TutorialRoot->LaunchTutorial(InTutorial, InStartType, InNavigationWindow, OnTutorialClosed, OnTutorialExited);
 	}
 }
 
diff --git a/Engine/Source/Editor/IntroTutorials/Private/IntroTutorials.h b/Engine/Source/Editor/IntroTutorials/Private/IntroTutorials.h
--- a/Engine/Source/Editor/IntroTutorials/Private/IntroTutorials.h
+++ b/Engine/Source/Editor/IntroTutorials/Private/IntroTutorials.h
@@ -104,7 +104,7 @@ public:
 
 	// IIntroTutorials interface
 	virtual void LaunchTutorial(const FString& TutorialAssetName) override;
-	virtual void LaunchTutorial(UEditorTutorial* Tutorial, bool bRestart = true, TWeakPtr<SWindow> InNavigationWindow = nullptr, FSimpleDelegate OnTutorialClosed = FSimpleDelegate(), FSimpleDelegate OnTutorialExited = FSimpleDelegate()) override;
+	virtual void LaunchTutorial(UEditorTutorial* Tutorial, IIntroTutorials::ETutorialStartType InStartType = IIntroTutorials::ETutorialStartType::TST_RESTART, TWeakPtr<SWindow> InNavigationWindow = nullptr, FSimpleDelegate OnTutorialClosed = FSimpleDelegate(), FSimpleDelegate OnTutorialExited = FSimpleDelegate()) override;
 	virtual void CloseAllTutorialContent() override;
 	virtual TSharedRef<SWidget> CreateTutorialsWidget(FName InContext, TWeakPtr<SWindow> InContextWindow = nullptr) const override;
 	virtual TSharedPtr<SWidget> CreateTutorialsLoadingWidget(TWeakPtr<SWindow> InContextWindow = nullptr) const override;
diff --git a/Engine/Source/Editor/IntroTutorials/Private/SEditorTutorials.cpp b/Engine/Source/Editor/IntroTutorials/Private/SEditorTutorials.cpp
--- a/Engine/Source/Editor/IntroTutorials/Private/SEditorTutorials.cpp
+++ b/Engine/Source/Editor/IntroTutorials/Private/SEditorTutorials.cpp
@@ -101,7 +101,9 @@ bool SEditorTutorials::IsBackButtonEnabled() const
 	UEditorTutorial* CurrentTutorial = OnGetCurrentTutorial.Execute();
 	int32 CurrentTutorialStage = OnGetCurrentTutorialStage.Execute();
 
-	return CurrentTutorial != nullptr && CurrentTutorialStage > 0;
+	const bool bStepsPassed = (CurrentTutorial != nullptr) && (CurrentTutorialStage > 0);
+	const bool bPreviousTutorialAvailable = (CurrentTutorial != nullptr) && (CurrentTutorialStage == 0) && (CurrentTutorial->PreviousTutorial != nullptr);
+	return bStepsPassed || bPreviousTutorialAvailable;
 }
 
 bool SEditorTutorials::IsHomeButtonEnabled() const
@@ -114,8 +116,8 @@ bool SEditorTutorials::IsNextButtonEnabled() const
 	UEditorTutorial* CurrentTutorial = OnGetCurrentTutorial.Execute();
 	int32 CurrentTutorialStage = OnGetCurrentTutorialStage.Execute();
 
-	const bool bStepsRemaining = CurrentTutorial != nullptr && CurrentTutorialStage + 1 < CurrentTutorial->Stages.Num();
-	const bool bNextTutorialAvailable = CurrentTutorial != nullptr && CurrentTutorialStage + 1 >= CurrentTutorial->Stages.Num() && CurrentTutorial->NextTutorial != nullptr;
+	const bool bStepsRemaining = (CurrentTutorial != nullptr) && (CurrentTutorialStage + 1 < CurrentTutorial->Stages.Num());
+	const bool bNextTutorialAvailable = (CurrentTutorial != nullptr) && (CurrentTutorialStage + 1 >= CurrentTutorial->Stages.Num()) && (CurrentTutorial->NextTutorial != nullptr);
 	return bStepsRemaining || bNextTutorialAvailable;
 }
 
diff --git a/Engine/Source/Editor/IntroTutorials/Private/STutorialButton.cpp b/Engine/Source/Editor/IntroTutorials/Private/STutorialButton.cpp
--- a/Engine/Source/Editor/IntroTutorials/Private/STutorialButton.cpp
+++ b/Engine/Source/Editor/IntroTutorials/Private/STutorialButton.cpp
@@ -64,7 +64,7 @@ EActiveTimerReturnType STutorialButton::OpenTutorialPostConstruct( double InCurr
 		// kick off the attract tutorial if the user hasn't dismissed it and hasn't completed it
 		FIntroTutorials& IntroTutorials = FModuleManager::GetModuleChecked<FIntroTutorials>( TEXT( "IntroTutorials" ) );
 		const bool bRestart = true;
-		IntroTutorials.LaunchTutorial(CachedAttractTutorial, bRestart, ContextWindow);
+		IntroTutorials.LaunchTutorial(CachedAttractTutorial, bRestart ? IIntroTutorials::ETutorialStartType::TST_RESTART : IIntroTutorials::ETutorialStartType::TST_CONTINUE, ContextWindow);
 	}
 
 	if ( ShouldShowAlert() )
@@ -195,8 +195,7 @@ EActiveTimerReturnType STutorialButton::HandleButtonClicked_AssetRegistryChecker
 		//If we don't want to launch the browser, and we have a tutorial in mind, launch the tutorial now.
 		auto Delegate = FSimpleDelegate::CreateSP(this, &STutorialButton::HandleTutorialExited);
 
-		const bool bRestart = true;
-		IntroTutorials.LaunchTutorial(CachedLaunchTutorial, bRestart, ContextWindow, Delegate, Delegate);
+		IntroTutorials.LaunchTutorial(CachedLaunchTutorial, IIntroTutorials::ETutorialStartType::TST_RESTART, ContextWindow, Delegate, Delegate);
 
 		// The user asked to start the tutorial, so we don't need to remind them about it again this session, but we'll remind
 		// them in the next session if they haven't completed it by then
diff --git a/Engine/Source/Editor/IntroTutorials/Private/STutorialContent.cpp b/Engine/Source/Editor/IntroTutorials/Private/STutorialContent.cpp
--- a/Engine/Source/Editor/IntroTutorials/Private/STutorialContent.cpp
+++ b/Engine/Source/Editor/IntroTutorials/Private/STutorialContent.cpp
@@ -675,8 +675,7 @@ void STutorialContent::HandleRestartSelected()
 	if(Tutorial.IsValid())
 	{
 		FIntroTutorials& IntroTutorials = FModuleManager::GetModuleChecked<FIntroTutorials>(TEXT("IntroTutorials"));
-		const bool bRestart = true;
-		IntroTutorials.LaunchTutorial(Tutorial.Get(), bRestart, FSlateApplication::Get().FindWidgetWindow(AsShared()));
+		IntroTutorials.LaunchTutorial(Tutorial.Get(), IIntroTutorials::ETutorialStartType::TST_RESTART, FSlateApplication::Get().FindWidgetWindow(AsShared()));
 
 		if( FEngineAnalytics::IsAvailable() && Tutorial.IsValid() )
 		{
diff --git a/Engine/Source/Editor/IntroTutorials/Private/STutorialRoot.cpp b/Engine/Source/Editor/IntroTutorials/Private/STutorialRoot.cpp
--- a/Engine/Source/Editor/IntroTutorials/Private/STutorialRoot.cpp
+++ b/Engine/Source/Editor/IntroTutorials/Private/STutorialRoot.cpp
@@ -98,20 +98,32 @@ void STutorialRoot::Tick( const FGeometry& AllottedGeometry, const double InCurr
 	DrawnWidgets.Empty(DrawnWidgets.Max());
 }
 
-void STutorialRoot::LaunchTutorial(UEditorTutorial* InTutorial, bool bInRestart, TWeakPtr<SWindow> InNavigationWindow, FSimpleDelegate InOnTutorialClosed, FSimpleDelegate InOnTutorialExited)
+void STutorialRoot::LaunchTutorial(UEditorTutorial* InTutorial, IIntroTutorials::ETutorialStartType InStartType, TWeakPtr<SWindow> InNavigationWindow, FSimpleDelegate InOnTutorialClosed, FSimpleDelegate InOnTutorialExited)
 {
 	if(InTutorial != nullptr)
 	{
 		CurrentTutorial = InTutorial;
 
 		// we force a restart if this tutorial was completed
-		if(GetDefault<UTutorialStateSettings>()->HaveCompletedTutorial(CurrentTutorial))
+		if (GetDefault<UTutorialStateSettings>()->HaveCompletedTutorial(CurrentTutorial) && (InStartType == IIntroTutorials::ETutorialStartType::TST_CONTINUE))
 		{
-			bInRestart = true;
+			InStartType = IIntroTutorials::ETutorialStartType::TST_RESTART;
 		}
 
 		bool bHaveSeenTutorial = false;
-		CurrentTutorialStage = bInRestart ? 0 : GetDefault<UTutorialStateSettings>()->GetProgress(CurrentTutorial, bHaveSeenTutorial);
+		switch (InStartType)
+		{
+		case IIntroTutorials::ETutorialStartType::TST_RESTART:
+			CurrentTutorialStage = 0;
+			break;
+		case IIntroTutorials::ETutorialStartType::TST_LASTSTAGE:
+			CurrentTutorialStage = FMath::Max(0, (CurrentTutorial->Stages.Num() - 1));
+			break;
+		default:
+		case IIntroTutorials::ETutorialStartType::TST_CONTINUE:
+			CurrentTutorialStage = GetDefault<UTutorialStateSettings>()->GetProgress(CurrentTutorial, bHaveSeenTutorial);
+			break;
+		}
 
 		// check if we should be launching this tutorial for an asset editor
 		if(InTutorial->AssetToUse.IsValid())
@@ -266,29 +278,55 @@ void STutorialRoot::AddReferencedObjects( FReferenceCollector& Collector )
 
 void STutorialRoot::GoToPreviousStage()
 {
-	if(CurrentTutorial != nullptr)
+	if (CurrentTutorial != nullptr)
 	{
-		int32 PreviousTutorialStage = CurrentTutorialStage;
-		if (CurrentTutorialStage > 0)
+		UEditorTutorial* OldTutorial = CurrentTutorial;
+		int32 OldTutorialStage = CurrentTutorialStage;
+		if (CurrentTutorialStage < CurrentTutorial->Stages.Num())
 		{
 			CurrentTutorial->HandleTutorialStageEnded(CurrentTutorial->Stages[CurrentTutorialStage].Name);
 		}
 
-		CurrentTutorialStage = FMath::Max(0, CurrentTutorialStage - 1);
+		if ((CurrentTutorialStage <= 0) && (FName(*CurrentTutorial->PreviousTutorial.AssetLongPathname) != NAME_None))
+		{
+			TSubclassOf<UEditorTutorial> PreviousTutorialClass = LoadClass<UEditorTutorial>(NULL, *CurrentTutorial->PreviousTutorial.AssetLongPathname, NULL, LOAD_None, NULL);
+			if (PreviousTutorialClass != nullptr)
+			{
+				LaunchTutorial(PreviousTutorialClass->GetDefaultObject<UEditorTutorial>(), IIntroTutorials::ETutorialStartType::TST_LASTSTAGE, nullptr, FSimpleDelegate(), FSimpleDelegate());
+			}
+			else
+			{
+				FSlateNotificationManager::Get().AddNotification(FNotificationInfo(FText::Format(LOCTEXT("TutorialNotFound", "Could not start previous tutorial {0}"), FText::FromString(CurrentTutorial->PreviousTutorial.AssetLongPathname))));
+			}
+		}
+		else
+		{
+			CurrentTutorialStage = FMath::Max(CurrentTutorialStage - 1, 0);
+			GetMutableDefault<UTutorialStateSettings>()->RecordProgress(CurrentTutorial, CurrentTutorialStage);
+		}
 
-		if (PreviousTutorialStage != CurrentTutorialStage)
+		if (CurrentTutorial != nullptr && CurrentTutorialStage < CurrentTutorial->Stages.Num() && (CurrentTutorial != OldTutorial || CurrentTutorialStage != OldTutorialStage))
 		{
 			CurrentTutorial->HandleTutorialStageStarted(CurrentTutorial->Stages[CurrentTutorialStage].Name);
 		}
 	}
+
+	for (auto& TutorialWidget : TutorialWidgets)
+	{
+		if (TutorialWidget.Value.IsValid())
+		{
+			TSharedPtr<SEditorTutorials> PinnedTutorialWidget = TutorialWidget.Value.Pin();
+			PinnedTutorialWidget->RebuildCurrentContent();
+		}
+	}
 }
 
 void STutorialRoot::GoToNextStage(TWeakPtr<SWindow> InNavigationWindow)
 {
 	if(CurrentTutorial != nullptr)
 	{
-		UEditorTutorial* PreviousTutorial = CurrentTutorial;
-		int32 PreviousTutorialStage = CurrentTutorialStage;
+		UEditorTutorial* OldTutorial = CurrentTutorial;
+		int32 OldTutorialStage = CurrentTutorialStage;
 		if(CurrentTutorialStage < CurrentTutorial->Stages.Num())
 		{
 			CurrentTutorial->HandleTutorialStageEnded(CurrentTutorial->Stages[CurrentTutorialStage].Name);
@@ -299,7 +337,7 @@ void STutorialRoot::GoToNextStage(TWeakPtr<SWindow> InNavigationWindow)
 			TSubclassOf<UEditorTutorial> NextTutorialClass = LoadClass<UEditorTutorial>(NULL, *CurrentTutorial->NextTutorial.AssetLongPathname, NULL, LOAD_None, NULL);
 			if(NextTutorialClass != nullptr)
 			{
-				LaunchTutorial(NextTutorialClass->GetDefaultObject<UEditorTutorial>(), true, InNavigationWindow, FSimpleDelegate(), FSimpleDelegate());
+				LaunchTutorial(NextTutorialClass->GetDefaultObject<UEditorTutorial>(), IIntroTutorials::ETutorialStartType::TST_RESTART, InNavigationWindow, FSimpleDelegate(), FSimpleDelegate());
 			}
 			else
 			{
@@ -312,7 +350,7 @@ void STutorialRoot::GoToNextStage(TWeakPtr<SWindow> InNavigationWindow)
 			GetMutableDefault<UTutorialStateSettings>()->RecordProgress(CurrentTutorial, CurrentTutorialStage);
 		}
 
-		if (CurrentTutorial != nullptr && CurrentTutorialStage < CurrentTutorial->Stages.Num() && (CurrentTutorial != PreviousTutorial || CurrentTutorialStage != PreviousTutorialStage))
+		if (CurrentTutorial != nullptr && CurrentTutorialStage < CurrentTutorial->Stages.Num() && (CurrentTutorial != OldTutorial || CurrentTutorialStage != OldTutorialStage))
 		{
 			CurrentTutorial->HandleTutorialStageStarted(CurrentTutorial->Stages[CurrentTutorialStage].Name);
 		}
diff --git a/Engine/Source/Editor/IntroTutorials/Private/STutorialRoot.h b/Engine/Source/Editor/IntroTutorials/Private/STutorialRoot.h
--- a/Engine/Source/Editor/IntroTutorials/Private/STutorialRoot.h
+++ b/Engine/Source/Editor/IntroTutorials/Private/STutorialRoot.h
@@ -28,7 +28,7 @@ public:
 	virtual void AddReferencedObjects( FReferenceCollector& Collector ) override;
 
 	/** Launch the specified tutorial from the specified window */
-	void LaunchTutorial(UEditorTutorial* InTutorial, bool bInRestart, TWeakPtr<SWindow> InNavigationWindow, FSimpleDelegate InOnTutorialClosed, FSimpleDelegate InOnTutorialExited);
+	void LaunchTutorial(UEditorTutorial* InTutorial, IIntroTutorials::ETutorialStartType InStartType, TWeakPtr<SWindow> InNavigationWindow, FSimpleDelegate InOnTutorialClosed, FSimpleDelegate InOnTutorialExited);
 
 	/** Close all tutorial content */
 	void CloseAllTutorialContent();
diff --git a/Engine/Source/Editor/IntroTutorials/Private/STutorialsBrowser.cpp b/Engine/Source/Editor/IntroTutorials/Private/STutorialsBrowser.cpp
--- a/Engine/Source/Editor/IntroTutorials/Private/STutorialsBrowser.cpp
+++ b/Engine/Source/Editor/IntroTutorials/Private/STutorialsBrowser.cpp
@@ -872,7 +872,7 @@ void STutorialsBrowser::OnTutorialSelected(UEditorTutorial* InTutorial, bool bRe
 			IntroTutorials.DismissTutorialBrowser();
 		}
 	}
-	OnLaunchTutorial.ExecuteIfBound(InTutorial, bRestart, ParentWindow, FSimpleDelegate(), FSimpleDelegate());
+	OnLaunchTutorial.ExecuteIfBound(InTutorial, bRestart ? IIntroTutorials::ETutorialStartType::TST_RESTART : IIntroTutorials::ETutorialStartType::TST_CONTINUE, ParentWindow, FSimpleDelegate(), FSimpleDelegate());
 }
 
 void STutorialsBrowser::OnCategorySelected(const FString& InCategory)
diff --git a/Engine/Source/Editor/IntroTutorials/Private/STutorialsBrowser.h b/Engine/Source/Editor/IntroTutorials/Private/STutorialsBrowser.h
--- a/Engine/Source/Editor/IntroTutorials/Private/STutorialsBrowser.h
+++ b/Engine/Source/Editor/IntroTutorials/Private/STutorialsBrowser.h
@@ -4,7 +4,7 @@
 
 class UEditorTutorial;
 
-DECLARE_DELEGATE_FiveParams(FOnLaunchTutorial, UEditorTutorial* /** InTutorialToLaunch */, bool /* bRestart */, TWeakPtr<SWindow> /* InFromWindow */, FSimpleDelegate /* InOnTutorialClosed */, FSimpleDelegate /* InOnTutorialExited */);
+DECLARE_DELEGATE_FiveParams(FOnLaunchTutorial, UEditorTutorial* /** InTutorialToLaunch */, IIntroTutorials::ETutorialStartType /* InStartType */, TWeakPtr<SWindow> /* InFromWindow */, FSimpleDelegate /* InOnTutorialClosed */, FSimpleDelegate /* InOnTutorialExited */);
 
 /** Abstract base class for list entries in the tutorial menu */
 struct ITutorialListEntry
diff --git a/Engine/Source/Editor/IntroTutorials/Private/TutorialText.cpp b/Engine/Source/Editor/IntroTutorials/Private/TutorialText.cpp
--- a/Engine/Source/Editor/IntroTutorials/Private/TutorialText.cpp
+++ b/Engine/Source/Editor/IntroTutorials/Private/TutorialText.cpp
@@ -101,8 +101,7 @@ static void ParseTutorialLink(const FString &InternalLink)
 	{
 		FIntroTutorials& IntroTutorials = FModuleManager::GetModuleChecked<FIntroTutorials>(TEXT("IntroTutorials"));
 		IMainFrameModule& MainFrameModule = FModuleManager::LoadModuleChecked<IMainFrameModule>(TEXT("MainFrame"));
-		const bool bRestart = true;
-		IntroTutorials.LaunchTutorial(Blueprint->GeneratedClass->GetDefaultObject<UEditorTutorial>(), bRestart, MainFrameModule.GetParentWindow());
+		IntroTutorials.LaunchTutorial(Blueprint->GeneratedClass->GetDefaultObject<UEditorTutorial>(), IIntroTutorials::ETutorialStartType::TST_RESTART, MainFrameModule.GetParentWindow());
 
 		if( FEngineAnalytics::IsAvailable() )
 		{
diff --git a/Engine/Source/Editor/IntroTutorials/Public/EditorTutorial.h b/Engine/Source/Editor/IntroTutorials/Public/EditorTutorial.h
--- a/Engine/Source/Editor/IntroTutorials/Public/EditorTutorial.h
+++ b/Engine/Source/Editor/IntroTutorials/Public/EditorTutorial.h
@@ -247,6 +247,10 @@ class INTROTUTORIALS_API UEditorTutorial : public UObject
 	UPROPERTY(EditAnywhere, Category="Stages")
 	TArray<FTutorialStage> Stages;
 
+	/** Tutorial to optionally chain back to if the "back" button is clicked on the first stage */
+	UPROPERTY(EditAnywhere, Category = "Tutorial", meta = (MetaClass = "EditorTutorial"))
+	FStringClassReference PreviousTutorial;
+
 	/** Tutorial to optionally chain onto after this tutorial completes */
 	UPROPERTY(EditAnywhere, Category="Tutorial", meta=(MetaClass="EditorTutorial"))
 	FStringClassReference NextTutorial;
diff --git a/Engine/Source/Editor/IntroTutorials/Public/IIntroTutorials.h b/Engine/Source/Editor/IntroTutorials/Public/IIntroTutorials.h
--- a/Engine/Source/Editor/IntroTutorials/Public/IIntroTutorials.h
+++ b/Engine/Source/Editor/IntroTutorials/Public/IIntroTutorials.h
@@ -45,6 +45,17 @@ public:
 	 */
 	virtual void LaunchTutorial(const FString& TutorialAssetName) = 0;
 
+	//Parameter for LaunchTutorial function
+	enum ETutorialStartType
+	{
+		//Default start type, picks up where user left off if appropriate
+		TST_CONTINUE,
+		//Restart and clear progress
+		TST_RESTART,
+		//Start at the last stage of the tutorial, used when backing up to a previous tutorial
+		TST_LASTSTAGE
+	};
+
 	/**
 	 * Launch a tutorial immediately, bypassing the tutorial browser.
 	 * 
@@ -52,7 +63,7 @@ public:
 	 * @param	bRestart	Whether to restart the tutorial or resume from where we left off last time.
 	 * @param	InNavigationWindow	Optional window to launch the tutorial from - this is where navigation will be displayed.
 	 */
-	virtual void LaunchTutorial(UEditorTutorial* Tutorial, bool bRestart = true, TWeakPtr<SWindow> InNavigationWindow = nullptr, FSimpleDelegate OnTutorialClosed = FSimpleDelegate(), FSimpleDelegate OnTutorialExited = FSimpleDelegate()) = 0;
+	virtual void LaunchTutorial(UEditorTutorial* Tutorial, IIntroTutorials::ETutorialStartType InStartType = TST_RESTART, TWeakPtr<SWindow> InNavigationWindow = nullptr, FSimpleDelegate OnTutorialClosed = FSimpleDelegate(), FSimpleDelegate OnTutorialExited = FSimpleDelegate()) = 0;
 
 	/**
  	 *  Close all tutorial content, including the browser.
