[UE-7857] Scroll bar fix for tutorials browser.

diff --git a/Engine/Source/Editor/IntroTutorials/Private/STutorialsBrowser.cpp b/Engine/Source/Editor/IntroTutorials/Private/STutorialsBrowser.cpp
--- a/Engine/Source/Editor/IntroTutorials/Private/STutorialsBrowser.cpp
+++ b/Engine/Source/Editor/IntroTutorials/Private/STutorialsBrowser.cpp
@@ -602,15 +602,11 @@ void STutorialsBrowser::Construct(const FArguments& InArgs)
 			+SVerticalBox::Slot()
 			.FillHeight(1.0f)
 			[
-				SNew(SScrollBox)
-				+SScrollBox::Slot()
-				[
-					SAssignNew(TutorialList, SListView<TSharedPtr<ITutorialListEntry>>)
-					.ItemHeight(128.0f)
-					.ListItemsSource(&FilteredEntries)
-					.OnGenerateRow(this, &STutorialsBrowser::OnGenerateTutorialRow)
-					.SelectionMode(ESelectionMode::None)
-				]
+				SAssignNew(TutorialList, SListView<TSharedPtr<ITutorialListEntry>>)
+				.ItemHeight(128.0f)
+				.ListItemsSource(&FilteredEntries)
+				.OnGenerateRow(this, &STutorialsBrowser::OnGenerateTutorialRow)
+				.SelectionMode(ESelectionMode::None)
 			]
 		]
 	];
