[UE-16409] Added an editor option to prevent tutorial alerts from popping up in the future, including in other projects.

diff --git a/Engine/Source/Editor/IntroTutorials/Private/EditorTutorialSettings.h b/Engine/Source/Editor/IntroTutorials/Private/EditorTutorialSettings.h
--- a/Engine/Source/Editor/IntroTutorials/Private/EditorTutorialSettings.h
+++ b/Engine/Source/Editor/IntroTutorials/Private/EditorTutorialSettings.h
@@ -34,6 +34,10 @@ class UEditorTutorialSettings : public UObject
 {
 	GENERATED_UCLASS_BODY()
 
+	/** Disable the pulsing alert that indicates a new tutorial is available. */
+	UPROPERTY(Config, EditAnywhere, Category = "Tutorials")
+	bool bDisableAllTutorialAlerts;
+
 	/** Categories for tutorials */
 	UPROPERTY(Config, EditAnywhere, Category="Tutorials")
 	TArray<FTutorialCategory> Categories;
diff --git a/Engine/Source/Editor/IntroTutorials/Private/STutorialButton.cpp b/Engine/Source/Editor/IntroTutorials/Private/STutorialButton.cpp
--- a/Engine/Source/Editor/IntroTutorials/Private/STutorialButton.cpp
+++ b/Engine/Source/Editor/IntroTutorials/Private/STutorialButton.cpp
@@ -317,7 +317,7 @@ bool STutorialButton::ShouldShowAlert() const
 {
 	if ((bTestAlerts || !FEngineBuildSettings::IsInternalBuild()) && bTutorialAvailable && !(bTutorialCompleted || bTutorialDismissed))
 	{
-		return !GetMutableDefault<UTutorialStateSettings>()->AreAllTutorialsDismissed();
+		return (!GetMutableDefault<UEditorTutorialSettings>()->bDisableAllTutorialAlerts && !GetMutableDefault<UTutorialStateSettings>()->AreAllTutorialsDismissed());
 	}
 	return false;
 }
