Added a feature to in-engine tutorials that enables platform-specific screen skipping.

diff --git a/Engine/Source/Editor/IntroTutorials/Private/STutorialRoot.cpp b/Engine/Source/Editor/IntroTutorials/Private/STutorialRoot.cpp
--- a/Engine/Source/Editor/IntroTutorials/Private/STutorialRoot.cpp
+++ b/Engine/Source/Editor/IntroTutorials/Private/STutorialRoot.cpp
@@ -14,6 +14,7 @@
 #include "Editor/MainFrame/Public/Interfaces/IMainFrameModule.h"
 #include "NotificationManager.h"
 #include "SNotificationList.h"
+#include "Kismet/GameplayStatics.h"
 
 #define LOCTEXT_NAMESPACE "STutorialRoot"
 
@@ -286,23 +287,44 @@ void STutorialRoot::GoToPreviousStage()
 		{
 			CurrentTutorial->HandleTutorialStageEnded(CurrentTutorial->Stages[CurrentTutorialStage].Name);
 		}
-
-		if ((CurrentTutorialStage <= 0) && (FName(*CurrentTutorial->PreviousTutorial.AssetLongPathname) != NAME_None))
+		
+		FString PlatformName = UGameplayStatics::GetPlatformName();
+		for (--CurrentTutorialStage; CurrentTutorialStage >= 0; --CurrentTutorialStage)
 		{
-			TSubclassOf<UEditorTutorial> PreviousTutorialClass = LoadClass<UEditorTutorial>(NULL, *CurrentTutorial->PreviousTutorial.AssetLongPathname, NULL, LOAD_None, NULL);
-			if (PreviousTutorialClass != nullptr)
+			bool bPlatformFoundInList = false;
+			for (FString PlatformToTest : CurrentTutorial->Stages[CurrentTutorialStage].PlatformsToTest)
 			{
-				LaunchTutorial(PreviousTutorialClass->GetDefaultObject<UEditorTutorial>(), IIntroTutorials::ETutorialStartType::TST_LASTSTAGE, nullptr, FSimpleDelegate(), FSimpleDelegate());
+				if (!PlatformName.Compare(PlatformToTest, ESearchCase::IgnoreCase))
+				{
+					bPlatformFoundInList = true;
+					break;
+				}
 			}
-			else
+			if (bPlatformFoundInList != CurrentTutorial->Stages[CurrentTutorialStage].bInvertPlatformTest)
 			{
-				FSlateNotificationManager::Get().AddNotification(FNotificationInfo(FText::Format(LOCTEXT("TutorialNotFound", "Could not start previous tutorial {0}"), FText::FromString(CurrentTutorial->PreviousTutorial.AssetLongPathname))));
+				// Skip this stage
+				continue;
 			}
+			// We hit a page that we don't want to skip. Record progress and stop looking.
+			GetMutableDefault<UTutorialStateSettings>()->RecordProgress(CurrentTutorial, CurrentTutorialStage);
+			break;
 		}
-		else
+		if (CurrentTutorialStage < 0)
 		{
-			CurrentTutorialStage = FMath::Max(CurrentTutorialStage - 1, 0);
-			GetMutableDefault<UTutorialStateSettings>()->RecordProgress(CurrentTutorial, CurrentTutorialStage);
+			// We went out of bounds for this tutorial, so see if we want to go to another one.
+			CurrentTutorialStage = 0;
+			if (FName(*CurrentTutorial->PreviousTutorial.AssetLongPathname) != NAME_None)
+			{
+				TSubclassOf<UEditorTutorial> PreviousTutorialClass = LoadClass<UEditorTutorial>(NULL, *CurrentTutorial->PreviousTutorial.AssetLongPathname, NULL, LOAD_None, NULL);
+				if (PreviousTutorialClass != nullptr)
+				{
+					LaunchTutorial(PreviousTutorialClass->GetDefaultObject<UEditorTutorial>(), IIntroTutorials::ETutorialStartType::TST_LASTSTAGE, nullptr, FSimpleDelegate(), FSimpleDelegate());
+				}
+				else
+				{
+					FSlateNotificationManager::Get().AddNotification(FNotificationInfo(FText::Format(LOCTEXT("TutorialNotFound", "Could not start previous tutorial {0}"), FText::FromString(CurrentTutorial->PreviousTutorial.AssetLongPathname))));
+				}
+			}
 		}
 
 		if (CurrentTutorial != nullptr && CurrentTutorialStage < CurrentTutorial->Stages.Num() && (CurrentTutorial != OldTutorial || CurrentTutorialStage != OldTutorialStage))
@@ -332,22 +354,43 @@ void STutorialRoot::GoToNextStage(TWeakPtr<SWindow> InNavigationWindow)
 			CurrentTutorial->HandleTutorialStageEnded(CurrentTutorial->Stages[CurrentTutorialStage].Name);
 		}
 
-		if(CurrentTutorialStage + 1 >= CurrentTutorial->Stages.Num() && FName(*CurrentTutorial->NextTutorial.AssetLongPathname) != NAME_None)
+		FString PlatformName = UGameplayStatics::GetPlatformName();
+		for (++CurrentTutorialStage; CurrentTutorialStage < CurrentTutorial->Stages.Num(); ++CurrentTutorialStage)
 		{
-			TSubclassOf<UEditorTutorial> NextTutorialClass = LoadClass<UEditorTutorial>(NULL, *CurrentTutorial->NextTutorial.AssetLongPathname, NULL, LOAD_None, NULL);
-			if(NextTutorialClass != nullptr)
+			bool bPlatformFoundInList = false;
+			for (FString PlatformToTest : CurrentTutorial->Stages[CurrentTutorialStage].PlatformsToTest)
 			{
-				LaunchTutorial(NextTutorialClass->GetDefaultObject<UEditorTutorial>(), IIntroTutorials::ETutorialStartType::TST_RESTART, InNavigationWindow, FSimpleDelegate(), FSimpleDelegate());
+				if (!PlatformName.Compare(PlatformToTest, ESearchCase::IgnoreCase))
+				{
+					bPlatformFoundInList = true;
+					break;
+				}
 			}
-			else
+			if (bPlatformFoundInList != CurrentTutorial->Stages[CurrentTutorialStage].bInvertPlatformTest)
 			{
-				FSlateNotificationManager::Get().AddNotification(FNotificationInfo(FText::Format(LOCTEXT("TutorialNotFound", "Could not start next tutorial {0}"), FText::FromString(CurrentTutorial->NextTutorial.AssetLongPathname))));
+				// Skip this stage
+				continue;
 			}
+			// We hit a page that we don't want to skip. Record progress and stop looking.
+			GetMutableDefault<UTutorialStateSettings>()->RecordProgress(CurrentTutorial, CurrentTutorialStage);
+			break;
 		}
-		else
+		if (CurrentTutorialStage + 1 >= CurrentTutorial->Stages.Num())
 		{
-			CurrentTutorialStage = FMath::Min(CurrentTutorialStage + 1, CurrentTutorial->Stages.Num() - 1);
-			GetMutableDefault<UTutorialStateSettings>()->RecordProgress(CurrentTutorial, CurrentTutorialStage);
+			// We went out of bounds for this tutorial, so see if we want to go to another one.
+			CurrentTutorialStage = CurrentTutorial->Stages.Num() - 1;
+			if (FName(*CurrentTutorial->PreviousTutorial.AssetLongPathname) != NAME_None)
+			{
+				TSubclassOf<UEditorTutorial> NextTutorialClass = LoadClass<UEditorTutorial>(NULL, *CurrentTutorial->NextTutorial.AssetLongPathname, NULL, LOAD_None, NULL);
+				if (NextTutorialClass != nullptr)
+				{
+					LaunchTutorial(NextTutorialClass->GetDefaultObject<UEditorTutorial>(), IIntroTutorials::ETutorialStartType::TST_RESTART, InNavigationWindow, FSimpleDelegate(), FSimpleDelegate());
+				}
+				else
+				{
+					FSlateNotificationManager::Get().AddNotification(FNotificationInfo(FText::Format(LOCTEXT("TutorialNotFound", "Could not start next tutorial {0}"), FText::FromString(CurrentTutorial->NextTutorial.AssetLongPathname))));
+				}
+			}
 		}
 
 		if (CurrentTutorial != nullptr && CurrentTutorialStage < CurrentTutorial->Stages.Num() && (CurrentTutorial != OldTutorial || CurrentTutorialStage != OldTutorialStage))
diff --git a/Engine/Source/Editor/IntroTutorials/Public/EditorTutorial.h b/Engine/Source/Editor/IntroTutorials/Public/EditorTutorial.h
--- a/Engine/Source/Editor/IntroTutorials/Public/EditorTutorial.h
+++ b/Engine/Source/Editor/IntroTutorials/Public/EditorTutorial.h
@@ -210,7 +210,15 @@ struct INTROTUTORIALS_API FTutorialStage
 	
 	/** Text to display on the back button */
 	UPROPERTY(EditAnywhere, Category = "Stage")
-	FText BackButtonText;	
+	FText BackButtonText;
+
+	/** List of platforms to test against. Meaning of test is determined by InvertPlatformTest. */
+	UPROPERTY(EditAnywhere, Category = "Stage")
+	TArray<FString> PlatformsToTest;
+
+	/** If false, stage will be skipped if running on any platform in PlatformsToTest. If true, the stage will be if not running on any platform in PlatformsToTest. */
+	UPROPERTY(EditAnywhere, Category = "Stage")
+	bool bInvertPlatformTest;
 };
 
 /** An asset used to build a stage-by-stage tutorial in the editor */
