[UE-7479] Fixed so that tutorial BPs made from other blueprints will still be recognized by the browser.

diff --git a/Engine/Source/Editor/IntroTutorials/Private/STutorialsBrowser.cpp b/Engine/Source/Editor/IntroTutorials/Private/STutorialsBrowser.cpp
--- a/Engine/Source/Editor/IntroTutorials/Private/STutorialsBrowser.cpp
+++ b/Engine/Source/Editor/IntroTutorials/Private/STutorialsBrowser.cpp
@@ -730,6 +730,7 @@ void STutorialsBrowser::RebuildTutorials(TSharedPtr<FTutorialListEntry_Category>
 	FARFilter Filter;
 	Filter.ClassNames.Add(UBlueprint::StaticClass()->GetFName());
 	Filter.bRecursiveClasses = true;
+	Filter.TagsAndValues.Add(TEXT("NativeParentClass"), FString::Printf(TEXT("%s'%s'"), *UClass::StaticClass()->GetName(), *UEditorTutorial::StaticClass()->GetPathName()));
 	Filter.TagsAndValues.Add(TEXT("ParentClass"), FString::Printf(TEXT("%s'%s'"), *UClass::StaticClass()->GetName(), *UEditorTutorial::StaticClass()->GetPathName()));
 
 	TArray<FAssetData> AssetData;
