[UE-17116] Modifying save/load for tutorials to enable some cross-engine-version consistency.

diff --git a/Engine/Source/Editor/IntroTutorials/Private/TutorialStateSettings.cpp b/Engine/Source/Editor/IntroTutorials/Private/TutorialStateSettings.cpp
--- a/Engine/Source/Editor/IntroTutorials/Private/TutorialStateSettings.cpp
+++ b/Engine/Source/Editor/IntroTutorials/Private/TutorialStateSettings.cpp
@@ -6,14 +6,13 @@
 UTutorialStateSettings::UTutorialStateSettings(const FObjectInitializer& ObjectInitializer)
 	: Super(ObjectInitializer)
 {
-	bDismissedAllTutorials = false;
 }
 
 void UTutorialStateSettings::PostInitProperties()
 {
 	Super::PostInitProperties();
 
-	for(const auto& Progress : TutorialsProgress)
+	for (const auto& Progress : TutorialsProgress)
 	{
 		TSubclassOf<UEditorTutorial> TutorialClass = LoadClass<UEditorTutorial>(NULL, *Progress.Tutorial.ToString(), NULL, LOAD_None, NULL);
 		if(TutorialClass != nullptr)
@@ -141,4 +140,22 @@ void UTutorialStateSettings::ClearProgress()
 	bDismissedAllTutorials = false;
 
 	SaveConfig();
-}
\ No newline at end of file
+}
+
+void UTutorialStateSettings::SaveConfig(uint64 Flags /*= CPF_Config*/, const TCHAR* Filename /*= nullptr*/, FConfigCacheIni* Config /*= GConfig*/)
+{
+	// This class saves to an engine config and a per-project config.
+	// The engine-level config is there so that the user won't be bothered by the same tutorial in multiple projects.
+	// The per-project config is there so that the user won't be bothered when changing engine versions.
+	// The only way to see a repeated tutorial is to upgrade engine versions and then start a new project.
+	Super::SaveConfig(Flags, Filename, Config);
+	Super::SaveConfig(Flags, *GEditorPerProjectIni, Config);
+}
+
+void UTutorialStateSettings::LoadConfig(UClass* ConfigClass /*= nullptr*/, const TCHAR* Filename /*= nullptr*/, uint32 PropagationFlags /*= UE4::LCPF_None*/, class UProperty* PropertyToLoad /*= NULL*/)
+{
+	// This will trigger the post-load-config callback twice, but this class doesn't use it.
+	// We load from both engine and project so that we'll catch leftover project settings after upgrading the engine, and so that starting a new project doesn't miss the engine settings.
+	Super::LoadConfig(ConfigClass, *GEditorPerProjectIni, PropagationFlags, PropertyToLoad);
+	Super::LoadConfig(ConfigClass, Filename, PropagationFlags, PropertyToLoad);
+}
diff --git a/Engine/Source/Editor/IntroTutorials/Private/TutorialStateSettings.h b/Engine/Source/Editor/IntroTutorials/Private/TutorialStateSettings.h
--- a/Engine/Source/Editor/IntroTutorials/Private/TutorialStateSettings.h
+++ b/Engine/Source/Editor/IntroTutorials/Private/TutorialStateSettings.h
@@ -71,6 +71,12 @@ class UTutorialStateSettings : public UObject
 	/** Returns true if user has dismissed tutorials */
 	bool AreAllTutorialsDismissed();
 
+public:
+	//~ From UObject
+	virtual void SaveConfig(uint64 Flags = CPF_Config, const TCHAR* Filename = nullptr, FConfigCacheIni* Config = GConfig) override;
+	virtual void LoadConfig(UClass* ConfigClass = nullptr, const TCHAR* Filename = nullptr, uint32 PropagationFlags = UE4::LCPF_None, class UProperty* PropertyToLoad = nullptr) override;
+	//~
+
 private:
 	/** Recorded progress */
 	TMap<UEditorTutorial*, FTutorialProgress> ProgressMap;
diff --git a/Engine/Source/Runtime/CoreUObject/Public/UObject/UObject.h b/Engine/Source/Runtime/CoreUObject/Public/UObject/UObject.h
--- a/Engine/Source/Runtime/CoreUObject/Public/UObject/UObject.h
+++ b/Engine/Source/Runtime/CoreUObject/Public/UObject/UObject.h
@@ -773,7 +773,7 @@ public:
 	 * @warning: Must be safe on class-default metaobjects.
 	 * !!may benefit from hierarchical propagation, deleting keys that match superclass...not sure what's best yet.
 	 */
-	void SaveConfig( uint64 Flags=CPF_Config, const TCHAR* Filename=NULL, FConfigCacheIni* Config=GConfig );
+	virtual void SaveConfig( uint64 Flags=CPF_Config, const TCHAR* Filename=NULL, FConfigCacheIni* Config=GConfig );
 
 	/**
 	 * Saves just the section(s) for this class into the default ini file for the class (with just the changes from base)
@@ -824,7 +824,7 @@ public:
 	 * @param	PropagationFlags	indicates how this call to LoadConfig should be propagated; expects a bitmask of UE4::ELoadConfigPropagationFlags values.
 	 * @param	PropertyToLoad		if specified, only the ini value for the specified property will be imported.
 	 */
-	void LoadConfig( UClass* ConfigClass=NULL, const TCHAR* Filename=NULL, uint32 PropagationFlags=UE4::LCPF_None, class UProperty* PropertyToLoad=NULL );
+	virtual void LoadConfig( UClass* ConfigClass=NULL, const TCHAR* Filename=NULL, uint32 PropagationFlags=UE4::LCPF_None, class UProperty* PropertyToLoad=NULL );
 
 	/**
 	 * Wrapper method for LoadConfig that is used when reloading the config data for objects at runtime which have already loaded their config data at least once.
