[UE-19587] Tutorials that the user starts do not cause the green flash to show up in future sessions as per user feedback.

diff --git a/Engine/Source/Editor/IntroTutorials/Private/STutorialButton.cpp b/Engine/Source/Editor/IntroTutorials/Private/STutorialButton.cpp
--- a/Engine/Source/Editor/IntroTutorials/Private/STutorialButton.cpp
+++ b/Engine/Source/Editor/IntroTutorials/Private/STutorialButton.cpp
@@ -197,9 +197,9 @@ EActiveTimerReturnType STutorialButton::HandleButtonClicked_AssetRegistryChecker
 
 		IntroTutorials.LaunchTutorial(CachedLaunchTutorial, IIntroTutorials::ETutorialStartType::TST_RESTART, ContextWindow, Delegate, Delegate);
 
-		// The user asked to start the tutorial, so we don't need to remind them about it again this session, but we'll remind
-		// them in the next session if they haven't completed it by then
-		const bool bDismissAcrossSessions = false;
+		// The user asked to start the tutorial, so we don't need to remind them about it again.
+		// We used to remind them in future sessions, but user preference is that we don't.
+		const bool bDismissAcrossSessions = true;
 		GetMutableDefault<UTutorialStateSettings>()->DismissTutorial(CachedLaunchTutorial, bDismissAcrossSessions);
 		GetMutableDefault<UTutorialStateSettings>()->SaveProgress();
 		bTutorialDismissed = true;
