Fix ensure starting tutorials in any editor project

diff --git a/Engine/Source/Editor/IntroTutorials/Private/STutorialsBrowser.cpp b/Engine/Source/Editor/IntroTutorials/Private/STutorialsBrowser.cpp
--- a/Engine/Source/Editor/IntroTutorials/Private/STutorialsBrowser.cpp
+++ b/Engine/Source/Editor/IntroTutorials/Private/STutorialsBrowser.cpp
@@ -26,6 +26,7 @@
 #include "Widgets/Input/SSearchBox.h"
 #include "Widgets/Navigation/SBreadcrumbTrail.h"
 #include "UObject/GCObject.h"
+#include "Slate/DeferredCleanupSlateBrush.h"
 
 #define LOCTEXT_NAMESPACE "TutorialsBrowser"
 
@@ -73,8 +74,8 @@ public:
 			if(Texture != nullptr)
 			{
 				FIntPoint TextureSize = Texture->GetImportedSize();
-				DynamicBrush = MakeShareable(new FSlateDynamicImageBrush(Texture, FVector2D((float)TextureSize.X, (float)TextureSize.Y), NAME_None));
-				SlateBrush = DynamicBrush.Get();		
+				DynamicBrush = FDeferredCleanupSlateBrush::CreateBrush(Texture, FVector2D((float)TextureSize.X, (float)TextureSize.Y));
+				SlateBrush = DynamicBrush->GetSlateBrush();
 			}
 		}	
 
@@ -94,10 +95,6 @@ public:
 
 	virtual ~FTutorialListEntry_Category()
 	{
-		if( DynamicBrush.IsValid() )
-		{
-			DynamicBrush->ReleaseResource();
-		}
 	}
 
 	virtual TSharedRef<ITableRow> OnGenerateTutorialRow(const TSharedRef<STableViewBase>& OwnerTable) const override
@@ -273,7 +270,7 @@ public:
 	const FSlateBrush* SlateBrush;
 
 	/** Dynamic brush from the texture specified by the user */
-	TSharedPtr<FSlateDynamicImageBrush> DynamicBrush;
+	TSharedPtr<FDeferredCleanupSlateBrush> DynamicBrush;
 };
 
 DECLARE_DELEGATE_TwoParams(FOnTutorialSelected, UEditorTutorial* /* InTutorial */, bool /* bRestart */ );
@@ -291,8 +288,8 @@ public:
 		if(Tutorial->Texture != nullptr)
 		{
 			FIntPoint TextureSize = Tutorial->Texture->GetImportedSize();
-			DynamicBrush = MakeShareable(new FSlateDynamicImageBrush(Tutorial->Texture, FVector2D((float)TextureSize.X, (float)TextureSize.Y), NAME_None));
-			SlateBrush = DynamicBrush.Get();
+			DynamicBrush = FDeferredCleanupSlateBrush::CreateBrush(Tutorial->Texture, FVector2D((float)TextureSize.X, (float)TextureSize.Y));
+			SlateBrush = DynamicBrush->GetSlateBrush();
 		}	
 		else if(Tutorial->Icon.Len() > 0)
 		{
@@ -307,10 +304,6 @@ public:
 
 	virtual ~FTutorialListEntry_Tutorial()
 	{
-		if( DynamicBrush.IsValid() )
-		{
-			DynamicBrush->ReleaseResource();
-		}
 	}
 
 	virtual TSharedRef<ITableRow> OnGenerateTutorialRow(const TSharedRef<STableViewBase>& OwnerTable) const override
@@ -527,7 +520,7 @@ public:
 	const FSlateBrush* SlateBrush;
 
 	/** Dynamic brush from the texture specified by the user */
-	TSharedPtr<FSlateDynamicImageBrush> DynamicBrush;
+	TSharedPtr<FDeferredCleanupSlateBrush> DynamicBrush;
 
 	/** Cached tutorial completion state */
 	mutable bool bHaveCompletedTutorial;
