#DEV-VR - Changing MLSDK missing callback to launch MagicLeap URL

diff --git a/Engine/Source/Editor/IntroTutorials/Private/IntroTutorials.cpp b/Engine/Source/Editor/IntroTutorials/Private/IntroTutorials.cpp
--- a/Engine/Source/Editor/IntroTutorials/Private/IntroTutorials.cpp
+++ b/Engine/Source/Editor/IntroTutorials/Private/IntroTutorials.cpp
@@ -378,7 +378,14 @@ void FIntroTutorials::HandleSDKNotInstalled(const FString& PlatformName, const F
 	}
 	else
 	{
-		IDocumentation::Get()->Open( InTutorialAsset );
+		if (PlatformName.Equals(TEXT("Lumin")))
+		{
+			IDocumentation::Get()->Open(TEXT("https://docs.unrealengine.com/en-us/Platforms/MagicLeap"));
+		}
+		else
+		{
+			IDocumentation::Get()->Open(InTutorialAsset);
+		}
 	}
 }
 
