Fixed editor tutorials, they were accidently disabled.

diff --git a/Engine/Source/Editor/IntroTutorials/Private/IntroTutorials.cpp b/Engine/Source/Editor/IntroTutorials/Private/IntroTutorials.cpp
--- a/Engine/Source/Editor/IntroTutorials/Private/IntroTutorials.cpp
+++ b/Engine/Source/Editor/IntroTutorials/Private/IntroTutorials.cpp
@@ -86,7 +86,11 @@ TSharedRef<FExtender> FIntroTutorials::AddSummonBlueprintTutorialsMenuExtender(c
 void FIntroTutorials::StartupModule()
 {
 	// This code can run with content commandlets. Slate is not initialized with commandlets and the below code will fail.
-	if (!bDisableTutorials && !IsRunningCommandlet() && !FApp::IsUnattended() && !FApp::CanEverRender() )
+	const bool bCommandlet = IsRunningCommandlet();
+	const bool bUnattended = FApp::IsUnattended();
+	const bool bCanEverRender = FApp::CanEverRender();
+
+	if (!bDisableTutorials && !bCommandlet && !bUnattended && bCanEverRender)
 	{
 		// Add tutorial for main frame opening
 		IMainFrameModule& MainFrameModule = FModuleManager::LoadModuleChecked<IMainFrameModule>(TEXT("MainFrame"));
