Converted menu: Intro Tutorials
#rnx

diff --git a/Engine/Source/Editor/IntroTutorials/IntroTutorials.Build.cs b/Engine/Source/Editor/IntroTutorials/IntroTutorials.Build.cs
--- a/Engine/Source/Editor/IntroTutorials/IntroTutorials.Build.cs
+++ b/Engine/Source/Editor/IntroTutorials/IntroTutorials.Build.cs
@@ -43,6 +43,7 @@ namespace UnrealBuildTool.Rules
 					"LevelEditor",
                     "AssetRegistry",
 					"Analytics",
+					"EditorMenus",
 				}
 			);
 
diff --git a/Engine/Source/Editor/IntroTutorials/Private/IntroTutorials.cpp b/Engine/Source/Editor/IntroTutorials/Private/IntroTutorials.cpp
--- a/Engine/Source/Editor/IntroTutorials/Private/IntroTutorials.cpp
+++ b/Engine/Source/Editor/IntroTutorials/Private/IntroTutorials.cpp
@@ -37,7 +37,7 @@
 #include "AssetToolsModule.h"
 #include "IClassTypeActions.h"
 #include "ClassTypeActions_EditorTutorial.h"
-
+#include "EditorMenuSubsystem.h"
 
 #define LOCTEXT_NAMESPACE "IntroTutorials"
 
@@ -94,10 +94,7 @@ void FIntroTutorials::StartupModule()
 		MainFrameModule.OnMainFrameSDKNotInstalled().AddRaw(this, &FIntroTutorials::HandleSDKNotInstalled);
 		
 		// Add menu option for level editor tutorial
-		MainMenuExtender = MakeShareable(new FExtender);
-		MainMenuExtender->AddMenuExtension("HelpBrowse", EExtensionHook::After, NULL, FMenuExtensionDelegate::CreateRaw(this, &FIntroTutorials::AddSummonTutorialsMenuExtension));
-		FLevelEditorModule& LevelEditorModule = FModuleManager::LoadModuleChecked<FLevelEditorModule>( "LevelEditor" );
-		LevelEditorModule.GetMenuExtensibilityManager()->AddExtender(MainMenuExtender);
+		RegisterSummonTutorialsMenuEntries();
 
 		// Add menu option to blueprint editor as well
 		FBlueprintEditorModule& BPEditorModule = FModuleManager::LoadModuleChecked<FBlueprintEditorModule>( "Kismet" );
@@ -163,6 +160,8 @@ void FIntroTutorials::StartupModule()
 
 void FIntroTutorials::ShutdownModule()
 {
+	UEditorMenuSubsystem::UnregisterOwner(this);
+
 	if (!bDisableTutorials && !IsRunningCommandlet())
 	{
 		FSourceCodeNavigation::AccessOnCompilerNotFound().RemoveAll( this );
@@ -218,15 +217,18 @@ void FIntroTutorials::ShutdownModule()
 	FGlobalTabmanager::Get()->UnregisterNomadTabSpawner(TEXT("TutorialsBrowser"));
 }
 
-void FIntroTutorials::AddSummonTutorialsMenuExtension(FMenuBuilder& MenuBuilder)
+void FIntroTutorials::RegisterSummonTutorialsMenuEntries()
 {
-	MenuBuilder.BeginSection("Tutorials", LOCTEXT("TutorialsLabel", "Tutorials"));
-	MenuBuilder.AddMenuEntry(
+	FEditorMenuOwnerScoped OwnerScoped(this);
+	UEditorMenu* HelpMenu = UEditorMenuSubsystem::Get()->ExtendMenu("LevelEditor.MainMenu.Help");
+	FEditorMenuSection& Section = HelpMenu->AddSection("Tutorials", LOCTEXT("TutorialsLabel", "Tutorials"), FEditorMenuInsert("HelpBrowse", EEditorMenuInsertType::After));
+	Section.AddEntry(FEditorMenuEntry::InitMenuEntry(
+		"Tutorials",
 		LOCTEXT("TutorialsMenuEntryTitle", "Tutorials"),
 		LOCTEXT("TutorialsMenuEntryToolTip", "Opens up introductory tutorials covering the basics of using the Unreal Engine 4 Editor."),
 		FSlateIcon(FEditorStyle::GetStyleSetName(), "LevelEditor.Tutorials"),
-		FUIAction(FExecuteAction::CreateRaw(this, &FIntroTutorials::SummonTutorialHome)));
-	MenuBuilder.EndSection();
+		FUIAction(FExecuteAction::CreateRaw(this, &FIntroTutorials::SummonTutorialHome))
+	));
 }
 
 void FIntroTutorials::AddSummonBlueprintTutorialsMenuExtension(FMenuBuilder& MenuBuilder, UObject* PrimaryObject)
diff --git a/Engine/Source/Editor/IntroTutorials/Private/IntroTutorials.h b/Engine/Source/Editor/IntroTutorials/Private/IntroTutorials.h
--- a/Engine/Source/Editor/IntroTutorials/Private/IntroTutorials.h
+++ b/Engine/Source/Editor/IntroTutorials/Private/IntroTutorials.h
@@ -63,8 +63,8 @@ private:
 	virtual void StartupModule() override;
 	virtual void ShutdownModule() override;
 
-	/** Add the menu extension for summoning the tutorial */
-	void AddSummonTutorialsMenuExtension(FMenuBuilder& MenuBuilder);
+	/** Register the menu for summoning the tutorial */
+	void RegisterSummonTutorialsMenuEntries();
 
 	/** Add a menu extender to summon context-sensitive Blueprints page */
 	TSharedRef<FExtender> AddSummonBlueprintTutorialsMenuExtender(const TSharedRef<FUICommandList> CommandList, const TArray<UObject*> EditingObjects) const;
