Fix crash when following tutorials that opened modes tab

diff --git a/Engine/Source/Editor/IntroTutorials/Private/STutorialsBrowser.cpp b/Engine/Source/Editor/IntroTutorials/Private/STutorialsBrowser.cpp
--- a/Engine/Source/Editor/IntroTutorials/Private/STutorialsBrowser.cpp
+++ b/Engine/Source/Editor/IntroTutorials/Private/STutorialsBrowser.cpp
@@ -34,6 +34,8 @@ class FTutorialListEntry_Tutorial;
 
 DECLARE_DELEGATE_OneParam(FOnCategorySelected, const FString& /* InCategory */);
 
+DEFINE_LOG_CATEGORY_STATIC(LogTutorials, Log, All);
+
 namespace TutorialBrowserConstants
 {
 	const float RefreshTimerInterval = 1.0f;
@@ -757,6 +759,9 @@ void STutorialsBrowser::RebuildTutorials(TSharedPtr<FTutorialListEntry_Category>
 	TArray<FAssetData> AssetData;
 	AssetRegistry.Get().GetAssets(Filter, AssetData);
 
+	TArray<FString> DeletedTabs;
+	DeletedTabs.Add("LevelEditorToolBox");
+
 	for (const auto& TutorialAsset : AssetData)
 	{
 		UBlueprint* Blueprint = LoadObject<UBlueprint>(nullptr, *TutorialAsset.ObjectPath.ToString());
@@ -767,7 +772,29 @@ void STutorialsBrowser::RebuildTutorials(TSharedPtr<FTutorialListEntry_Category>
 			//UEditorTutorial* Tutorial = Blueprint->GeneratedClass->GetDefaultObject<UEditorTutorial>();
 			if(!Tutorial->bHideInBrowser)
 			{
-				Tutorials.Add(MakeShareable(new FTutorialListEntry_Tutorial(Tutorial, FOnTutorialSelected::CreateSP(this, &STutorialsBrowser::OnTutorialSelected), TAttribute<FText>::Create(TAttribute<FText>::FGetter::CreateSP(this, &STutorialsBrowser::GetSearchText)))));
+				bool bUsesDeletedTab = false;
+				for (const FTutorialStage& Stage : Tutorial->Stages)
+				{
+					for (const FTutorialWidgetContent& TutorialWidgetContent : Stage.WidgetContent)
+					{
+						if (DeletedTabs.Contains(TutorialWidgetContent.WidgetAnchor.TabToFocusOrOpen))
+						{
+							bUsesDeletedTab = true;
+							UE_LOG(LogTutorials, Warning, TEXT("Tutorial uses deleted tab: %s, tutorial: %s"), *TutorialWidgetContent.WidgetAnchor.TabToFocusOrOpen, *Blueprint->GetFullName());
+							break;
+						}
+					}
+
+					if (bUsesDeletedTab)
+					{
+						break;
+					}
+				}
+
+				if (!bUsesDeletedTab)
+				{
+					Tutorials.Add(MakeShareable(new FTutorialListEntry_Tutorial(Tutorial, FOnTutorialSelected::CreateSP(this, &STutorialsBrowser::OnTutorialSelected), TAttribute<FText>::Create(TAttribute<FText>::FGetter::CreateSP(this, &STutorialsBrowser::GetSearchText)))));
+				}
 			}
 		}
 	}
